Copying //UE4/Partner-MagicLeap-Staging to Dev-VR (//UE4/Dev-VR)
Magic Leap 0.22 SDK integration

diff --git a/Engine/Config/BaseEngine.ini b/Engine/Config/BaseEngine.ini
--- a/Engine/Config/BaseEngine.ini
+++ b/Engine/Config/BaseEngine.ini
@@ -2045,6 +2045,7 @@ NDKAPILevel=android-19
 [/Script/MagicLeap.MagicLeapSettings]
 bEnableZI=False
 bUseVulkanForZI=True
+bUseMLAudioForZI=True
 
 [/Script/LuminRuntimeSettings.LuminRuntimeSettings]
 PackageName=com.YourCompany.[PROJECT]
@@ -2052,14 +2053,14 @@ ApplicationDisplayName=
 bIsScreensApp=false
 FrameTimingHint=FPS_60
 bProtectedContent=False
+bManualCallToAppReady=False
 bUseMobileRendering=true
-bBuildWithNvTegraGfxDebugger=False
 Certificate=(FilePath="")
 IconModelPath=(Path="Build/Lumin/Resources/Model")
 IconPortalPath=(Path="Build/Lumin/Resources/Portal")
 bUseVulkan=true
 VersionCode=1
-MinimumAPILevel=2
+MinimumAPILevel=6
 ControllerTrackingType=Tracked
 ControllerTrackingMode=CoordinateFrameUID
 +AppPrivileges=LowLatencyLightwear
@@ -2073,10 +2074,10 @@ ControllerTrackingMode=CoordinateFrameUID
 +AppPrivileges=ScreensProvider
 +AppPrivileges=AudioCaptureMic
 +AppPrivileges=IdentityRead
-+AppPrivileges=AudioRecognizer
 +AppPrivileges=VoiceInput
 +AppPrivileges=ControllerPose
 bRemoveDebugInfo=true
+VulkanValidationLayerLibs=(Path="")
 
 [/Script/UnrealEd.CookerSettings]
 DefaultPVRTCQuality=1
diff --git a/Engine/Plugins/Lumin/MLSDK/MLSDK.uplugin b/Engine/Plugins/Lumin/MLSDK/MLSDK.uplugin
new file mode 100644
--- /dev/null
+++ b/Engine/Plugins/Lumin/MLSDK/MLSDK.uplugin
@@ -0,0 +1,18 @@
+{
+	"FileVersion" : 3,
+	"Version" : 1,
+	"VersionName": "1.0",
+	"FriendlyName" : "MLSDK",
+	"CreatedBy" : "Magic Leap, Inc.",
+	"CreatedByURL" : "http://magicleap.com",
+	"Description" : "Native C-API support for Magic Leap.",
+	"Category" : "Lumin",
+	"EnabledByDefault" : true,
+	"DocsURL" : "",
+	"MarketplaceURL" : "",
+	"SupportURL" : "",
+	"CanContainContent" : false,
+	"IsBetaVersion" : false,
+	"Installed": false,
+	"Modules": []
+}
diff --git a/Engine/Plugins/Lumin/MLSDK/Source/MLSDK/MLSDK.Build.cs b/Engine/Plugins/Lumin/MLSDK/Source/MLSDK/MLSDK.Build.cs
new file mode 100644
--- /dev/null
+++ b/Engine/Plugins/Lumin/MLSDK/Source/MLSDK/MLSDK.Build.cs
@@ -0,0 +1,64 @@
+// Copyright 1998-2019 Epic Games, Inc. All Rights Reserved.
+
+using UnrealBuildTool;
+using System.IO;
+using Tools.DotNETCommon;
+
+public class MLSDK : ModuleRules
+{
+	public MLSDK(ReadOnlyTargetRules Target) : base(Target)
+	{
+		// Needed for FVector, FQuat and FTransform used in MagicLeapMath.h
+		PrivateDependencyModuleNames.Add("Core");
+		// Include headers to be public to other modules.
+		PublicIncludePaths.Add(Path.Combine(ModuleDirectory, "Public"));
+
+		Type = ModuleType.External;
+
+		LuminSDKVersionHelper SDKVersionHelper = new LuminSDKVersionHelper();
+		string MLSDKPath = System.Environment.GetEnvironmentVariable("MLSDK");
+		bool bIsMLSDKInstalled = false;
+		bool bIsSupportedTargetPlatformAndType = ((Target.Platform == UnrealTargetPlatform.Lumin) || (Target.Type == TargetRules.TargetType.Editor));
+		if (SDKVersionHelper.HasAnySDK() && !string.IsNullOrEmpty(MLSDKPath))
+		{
+			string IncludePath = Path.Combine(MLSDKPath, "include");
+			string LibraryPath = Path.Combine(MLSDKPath, "lib");
+			string LibraryPlatformFolder = string.Empty;
+			if (Target.Platform == UnrealTargetPlatform.Win64)
+			{
+				LibraryPlatformFolder = "win64";
+			}
+			else if (Target.Platform == UnrealTargetPlatform.Mac)
+			{
+				LibraryPlatformFolder = "osx";
+			}
+			else if (Target.Platform == UnrealTargetPlatform.Linux)
+			{
+				LibraryPlatformFolder = "linux64";
+			}
+			else if (Target.Platform == UnrealTargetPlatform.Lumin)
+			{
+				LibraryPlatformFolder = "lumin";
+			}
+			else
+			{
+				// This will fail the bIsMLSDKInstalled check, causing WITH_MLSDK to be set to 0 for unsupported platforms.
+				LibraryPlatformFolder = "unsupported";
+			}
+
+			LibraryPath = Path.Combine(LibraryPath, LibraryPlatformFolder);
+
+			bIsMLSDKInstalled = Directory.Exists(IncludePath) && Directory.Exists(LibraryPath);
+			if (bIsMLSDKInstalled)
+			{
+				PublicIncludePaths.Add(IncludePath);
+				if (Target.Platform == UnrealTargetPlatform.Lumin)
+				{
+					PublicIncludePaths.Add(Target.UEThirdPartySourceDirectory + "Vulkan/Include/vulkan");
+				}
+			}
+		}
+
+		PublicDefinitions.Add("WITH_MLSDK=" + (bIsMLSDKInstalled && bIsSupportedTargetPlatformAndType ? "1" : "0"));
+	}
+}
diff --git a/Engine/Plugins/Lumin/MLSDK/Source/MLSDK/Public/Lumin/CAPIShims/LuminAPIAudio.h b/Engine/Plugins/Lumin/MLSDK/Source/MLSDK/Public/Lumin/CAPIShims/LuminAPIAudio.h
new file mode 100644
--- /dev/null
+++ b/Engine/Plugins/Lumin/MLSDK/Source/MLSDK/Public/Lumin/CAPIShims/LuminAPIAudio.h
@@ -0,0 +1,165 @@
+// Copyright 1998-2019 Epic Games, Inc. All Rights Reserved.
+
+#pragma once
+
+#if !defined(WITH_MLSDK) || WITH_MLSDK
+
+#include "Lumin/CAPIShims/LuminAPI.h"
+
+LUMIN_THIRD_PARTY_INCLUDES_START
+#include <ml_audio.h>
+LUMIN_THIRD_PARTY_INCLUDES_END
+
+namespace MLSDK_API
+{
+
+CREATE_FUNCTION_SHIM(ml_audio, const char*, MLAudioGetResultString)
+#define MLAudioGetResultString ::MLSDK_API::MLAudioGetResultStringShim
+CREATE_FUNCTION_SHIM(ml_audio, MLResult, MLAudioCreateLoadedResource)
+#define MLAudioCreateLoadedResource ::MLSDK_API::MLAudioCreateLoadedResourceShim
+CREATE_FUNCTION_SHIM(ml_audio, MLResult, MLAudioCreateStreamedResource)
+#define MLAudioCreateStreamedResource ::MLSDK_API::MLAudioCreateStreamedResourceShim
+CREATE_FUNCTION_SHIM(ml_audio, MLResult, MLAudioCheckResource)
+#define MLAudioCheckResource ::MLSDK_API::MLAudioCheckResourceShim
+CREATE_FUNCTION_SHIM(ml_audio, MLResult, MLAudioRefreshResource)
+#define MLAudioRefreshResource ::MLSDK_API::MLAudioRefreshResourceShim
+CREATE_FUNCTION_SHIM(ml_audio, MLResult, MLAudioDestroyResource)
+#define MLAudioDestroyResource ::MLSDK_API::MLAudioDestroyResourceShim
+CREATE_FUNCTION_SHIM(ml_audio, MLResult, MLAudioCreateSoundWithLoadedResource)
+#define MLAudioCreateSoundWithLoadedResource ::MLSDK_API::MLAudioCreateSoundWithLoadedResourceShim
+CREATE_FUNCTION_SHIM(ml_audio, MLResult, MLAudioCreateSoundWithStreamedResource)
+#define MLAudioCreateSoundWithStreamedResource ::MLSDK_API::MLAudioCreateSoundWithStreamedResourceShim
+CREATE_FUNCTION_SHIM(ml_audio, MLResult, MLAudioCreateSoundWithLoadedFile)
+#define MLAudioCreateSoundWithLoadedFile ::MLSDK_API::MLAudioCreateSoundWithLoadedFileShim
+CREATE_FUNCTION_SHIM(ml_audio, MLResult, MLAudioCreateSoundWithStreamedFile)
+#define MLAudioCreateSoundWithStreamedFile ::MLSDK_API::MLAudioCreateSoundWithStreamedFileShim
+CREATE_FUNCTION_SHIM(ml_audio, MLResult, MLAudioCreateSoundWithOutputStream)
+#define MLAudioCreateSoundWithOutputStream ::MLSDK_API::MLAudioCreateSoundWithOutputStreamShim
+CREATE_FUNCTION_SHIM(ml_audio, MLResult, MLAudioDestroySound)
+#define MLAudioDestroySound ::MLSDK_API::MLAudioDestroySoundShim
+CREATE_FUNCTION_SHIM(ml_audio, MLResult, MLAudioStartSound)
+#define MLAudioStartSound ::MLSDK_API::MLAudioStartSoundShim
+CREATE_FUNCTION_SHIM(ml_audio, MLResult, MLAudioStopSound)
+#define MLAudioStopSound ::MLSDK_API::MLAudioStopSoundShim
+CREATE_FUNCTION_SHIM(ml_audio, MLResult, MLAudioPauseSound)
+#define MLAudioPauseSound ::MLSDK_API::MLAudioPauseSoundShim
+CREATE_FUNCTION_SHIM(ml_audio, MLResult, MLAudioResumeSound)
+#define MLAudioResumeSound ::MLSDK_API::MLAudioResumeSoundShim
+CREATE_FUNCTION_SHIM(ml_audio, MLResult, MLAudioGetSoundState)
+#define MLAudioGetSoundState ::MLSDK_API::MLAudioGetSoundStateShim
+CREATE_FUNCTION_SHIM(ml_audio, MLResult, MLAudioGetSoundFormat)
+#define MLAudioGetSoundFormat ::MLSDK_API::MLAudioGetSoundFormatShim
+CREATE_FUNCTION_SHIM(ml_audio, MLResult, MLAudioSetSoundEventCallback)
+#define MLAudioSetSoundEventCallback ::MLSDK_API::MLAudioSetSoundEventCallbackShim
+CREATE_FUNCTION_SHIM(ml_audio, MLResult, MLAudioSetSpatialSoundEnable)
+#define MLAudioSetSpatialSoundEnable ::MLSDK_API::MLAudioSetSpatialSoundEnableShim
+CREATE_FUNCTION_SHIM(ml_audio, MLResult, MLAudioGetSpatialSoundEnable)
+#define MLAudioGetSpatialSoundEnable ::MLSDK_API::MLAudioGetSpatialSoundEnableShim
+CREATE_FUNCTION_SHIM(ml_audio, MLResult, MLAudioSetSpatialSoundPosition)
+#define MLAudioSetSpatialSoundPosition ::MLSDK_API::MLAudioSetSpatialSoundPositionShim
+CREATE_FUNCTION_SHIM(ml_audio, MLResult, MLAudioGetSpatialSoundPosition)
+#define MLAudioGetSpatialSoundPosition ::MLSDK_API::MLAudioGetSpatialSoundPositionShim
+CREATE_FUNCTION_SHIM(ml_audio, MLResult, MLAudioSetSpatialSoundDirection)
+#define MLAudioSetSpatialSoundDirection ::MLSDK_API::MLAudioSetSpatialSoundDirectionShim
+CREATE_FUNCTION_SHIM(ml_audio, MLResult, MLAudioGetSpatialSoundDirection)
+#define MLAudioGetSpatialSoundDirection ::MLSDK_API::MLAudioGetSpatialSoundDirectionShim
+CREATE_FUNCTION_SHIM(ml_audio, MLResult, MLAudioSetSpatialSoundDistanceProperties)
+#define MLAudioSetSpatialSoundDistanceProperties ::MLSDK_API::MLAudioSetSpatialSoundDistancePropertiesShim
+CREATE_FUNCTION_SHIM(ml_audio, MLResult, MLAudioGetSpatialSoundDistanceProperties)
+#define MLAudioGetSpatialSoundDistanceProperties ::MLSDK_API::MLAudioGetSpatialSoundDistancePropertiesShim
+CREATE_FUNCTION_SHIM(ml_audio, MLResult, MLAudioSetSpatialSoundRadiationProperties)
+#define MLAudioSetSpatialSoundRadiationProperties ::MLSDK_API::MLAudioSetSpatialSoundRadiationPropertiesShim
+CREATE_FUNCTION_SHIM(ml_audio, MLResult, MLAudioGetSpatialSoundRadiationProperties)
+#define MLAudioGetSpatialSoundRadiationProperties ::MLSDK_API::MLAudioGetSpatialSoundRadiationPropertiesShim
+CREATE_FUNCTION_SHIM(ml_audio, MLResult, MLAudioSetSpatialSoundDirectSendLevels)
+#define MLAudioSetSpatialSoundDirectSendLevels ::MLSDK_API::MLAudioSetSpatialSoundDirectSendLevelsShim
+CREATE_FUNCTION_SHIM(ml_audio, MLResult, MLAudioGetSpatialSoundDirectSendLevels)
+#define MLAudioGetSpatialSoundDirectSendLevels ::MLSDK_API::MLAudioGetSpatialSoundDirectSendLevelsShim
+CREATE_FUNCTION_SHIM(ml_audio, MLResult, MLAudioSetSpatialSoundRoomSendLevels)
+#define MLAudioSetSpatialSoundRoomSendLevels ::MLSDK_API::MLAudioSetSpatialSoundRoomSendLevelsShim
+CREATE_FUNCTION_SHIM(ml_audio, MLResult, MLAudioGetSpatialSoundRoomSendLevels)
+#define MLAudioGetSpatialSoundRoomSendLevels ::MLSDK_API::MLAudioGetSpatialSoundRoomSendLevelsShim
+CREATE_FUNCTION_SHIM(ml_audio, MLResult, MLAudioSetSpatialSoundRoomProperties)
+#define MLAudioSetSpatialSoundRoomProperties ::MLSDK_API::MLAudioSetSpatialSoundRoomPropertiesShim
+CREATE_FUNCTION_SHIM(ml_audio, MLResult, MLAudioGetSpatialSoundRoomProperties)
+#define MLAudioGetSpatialSoundRoomProperties ::MLSDK_API::MLAudioGetSpatialSoundRoomPropertiesShim
+CREATE_FUNCTION_SHIM(ml_audio, MLResult, MLAudioSetSpatialSoundControlFrequencies)
+#define MLAudioSetSpatialSoundControlFrequencies ::MLSDK_API::MLAudioSetSpatialSoundControlFrequenciesShim
+CREATE_FUNCTION_SHIM(ml_audio, MLResult, MLAudioGetSpatialSoundControlFrequencies)
+#define MLAudioGetSpatialSoundControlFrequencies ::MLSDK_API::MLAudioGetSpatialSoundControlFrequenciesShim
+CREATE_FUNCTION_SHIM(ml_audio, MLResult, MLAudioSetSoundVolumeLinear)
+#define MLAudioSetSoundVolumeLinear ::MLSDK_API::MLAudioSetSoundVolumeLinearShim
+CREATE_FUNCTION_SHIM(ml_audio, MLResult, MLAudioGetSoundVolumeLinear)
+#define MLAudioGetSoundVolumeLinear ::MLSDK_API::MLAudioGetSoundVolumeLinearShim
+CREATE_FUNCTION_SHIM(ml_audio, MLResult, MLAudioSetSoundVolumeDb)
+#define MLAudioSetSoundVolumeDb ::MLSDK_API::MLAudioSetSoundVolumeDbShim
+CREATE_FUNCTION_SHIM(ml_audio, MLResult, MLAudioGetSoundVolumeDb)
+#define MLAudioGetSoundVolumeDb ::MLSDK_API::MLAudioGetSoundVolumeDbShim
+CREATE_FUNCTION_SHIM(ml_audio, MLResult, MLAudioSetSoundPitch)
+#define MLAudioSetSoundPitch ::MLSDK_API::MLAudioSetSoundPitchShim
+CREATE_FUNCTION_SHIM(ml_audio, MLResult, MLAudioGetSoundPitch)
+#define MLAudioGetSoundPitch ::MLSDK_API::MLAudioGetSoundPitchShim
+CREATE_FUNCTION_SHIM(ml_audio, MLResult, MLAudioSetSoundMute)
+#define MLAudioSetSoundMute ::MLSDK_API::MLAudioSetSoundMuteShim
+CREATE_FUNCTION_SHIM(ml_audio, MLResult, MLAudioIsSoundMuted)
+#define MLAudioIsSoundMuted ::MLSDK_API::MLAudioIsSoundMutedShim
+CREATE_FUNCTION_SHIM(ml_audio, MLResult, MLAudioSetSoundLooping)
+#define MLAudioSetSoundLooping ::MLSDK_API::MLAudioSetSoundLoopingShim
+CREATE_FUNCTION_SHIM(ml_audio, MLResult, MLAudioIsSoundLooping)
+#define MLAudioIsSoundLooping ::MLSDK_API::MLAudioIsSoundLoopingShim
+CREATE_FUNCTION_SHIM(ml_audio, MLResult, MLAudioSetStreamedFileOffset)
+#define MLAudioSetStreamedFileOffset ::MLSDK_API::MLAudioSetStreamedFileOffsetShim
+CREATE_FUNCTION_SHIM(ml_audio, MLResult, MLAudioGetStreamedFileOffset)
+#define MLAudioGetStreamedFileOffset ::MLSDK_API::MLAudioGetStreamedFileOffsetShim
+CREATE_FUNCTION_SHIM(ml_audio, MLResult, MLAudioGetOutputDevice)
+#define MLAudioGetOutputDevice ::MLSDK_API::MLAudioGetOutputDeviceShim
+CREATE_FUNCTION_SHIM(ml_audio, MLResult, MLAudioGetOutputStreamDefaults)
+#define MLAudioGetOutputStreamDefaults ::MLSDK_API::MLAudioGetOutputStreamDefaultsShim
+CREATE_FUNCTION_SHIM(ml_audio, MLResult, MLAudioGetOutputStreamLatency)
+#define MLAudioGetOutputStreamLatency ::MLSDK_API::MLAudioGetOutputStreamLatencyShim
+CREATE_FUNCTION_SHIM(ml_audio, MLResult, MLAudioGetOutputStreamFramesPlayed)
+#define MLAudioGetOutputStreamFramesPlayed ::MLSDK_API::MLAudioGetOutputStreamFramesPlayedShim
+CREATE_FUNCTION_SHIM(ml_audio, MLResult, MLAudioGetOutputStreamBuffer)
+#define MLAudioGetOutputStreamBuffer ::MLSDK_API::MLAudioGetOutputStreamBufferShim
+CREATE_FUNCTION_SHIM(ml_audio, MLResult, MLAudioReleaseOutputStreamBuffer)
+#define MLAudioReleaseOutputStreamBuffer ::MLSDK_API::MLAudioReleaseOutputStreamBufferShim
+CREATE_FUNCTION_SHIM(ml_audio, MLResult, MLAudioGetMasterVolume)
+#define MLAudioGetMasterVolume ::MLSDK_API::MLAudioGetMasterVolumeShim
+CREATE_FUNCTION_SHIM(ml_audio, MLResult, MLAudioSetMasterVolumeCallback)
+#define MLAudioSetMasterVolumeCallback ::MLSDK_API::MLAudioSetMasterVolumeCallbackShim
+CREATE_FUNCTION_SHIM(ml_audio, MLResult, MLAudioCreateInputFromVoiceComm)
+#define MLAudioCreateInputFromVoiceComm ::MLSDK_API::MLAudioCreateInputFromVoiceCommShim
+CREATE_FUNCTION_SHIM(ml_audio, MLResult, MLAudioCreateInputFromWorldCapture)
+#define MLAudioCreateInputFromWorldCapture ::MLSDK_API::MLAudioCreateInputFromWorldCaptureShim
+CREATE_FUNCTION_SHIM(ml_audio, MLResult, MLAudioDestroyInput)
+#define MLAudioDestroyInput ::MLSDK_API::MLAudioDestroyInputShim
+CREATE_FUNCTION_SHIM(ml_audio, MLResult, MLAudioStartInput)
+#define MLAudioStartInput ::MLSDK_API::MLAudioStartInputShim
+CREATE_FUNCTION_SHIM(ml_audio, MLResult, MLAudioStopInput)
+#define MLAudioStopInput ::MLSDK_API::MLAudioStopInputShim
+CREATE_FUNCTION_SHIM(ml_audio, MLResult, MLAudioGetInputState)
+#define MLAudioGetInputState ::MLSDK_API::MLAudioGetInputStateShim
+CREATE_FUNCTION_SHIM(ml_audio, MLResult, MLAudioSetInputEventCallback)
+#define MLAudioSetInputEventCallback ::MLSDK_API::MLAudioSetInputEventCallbackShim
+CREATE_FUNCTION_SHIM(ml_audio, MLResult, MLAudioGetInputStreamDefaults)
+#define MLAudioGetInputStreamDefaults ::MLSDK_API::MLAudioGetInputStreamDefaultsShim
+CREATE_FUNCTION_SHIM(ml_audio, MLResult, MLAudioGetInputStreamLatency)
+#define MLAudioGetInputStreamLatency ::MLSDK_API::MLAudioGetInputStreamLatencyShim
+CREATE_FUNCTION_SHIM(ml_audio, MLResult, MLAudioGetInputStreamBuffer)
+#define MLAudioGetInputStreamBuffer ::MLSDK_API::MLAudioGetInputStreamBufferShim
+CREATE_FUNCTION_SHIM(ml_audio, MLResult, MLAudioReleaseInputStreamBuffer)
+#define MLAudioReleaseInputStreamBuffer ::MLSDK_API::MLAudioReleaseInputStreamBufferShim
+CREATE_FUNCTION_SHIM(ml_audio, MLResult, MLAudioSetMicMute)
+#define MLAudioSetMicMute ::MLSDK_API::MLAudioSetMicMuteShim
+CREATE_FUNCTION_SHIM(ml_audio, MLResult, MLAudioIsMicMuted)
+#define MLAudioIsMicMuted ::MLSDK_API::MLAudioIsMicMutedShim
+CREATE_FUNCTION_SHIM(ml_audio, MLResult, MLAudioSetMicMuteCallback)
+#define MLAudioSetMicMuteCallback ::MLSDK_API::MLAudioSetMicMuteCallbackShim
+CREATE_DEPRECATED_SHIM(ml_audio, MLResult, MLAudioCreateInputFromVirtualCapture)
+#define MLAudioCreateInputFromVirtualCapture ::MLSDK_API::MLAudioCreateInputFromVirtualCaptureShim
+CREATE_DEPRECATED_SHIM(ml_audio, MLResult, MLAudioCreateInputFromMixedCapture)
+#define MLAudioCreateInputFromMixedCapture ::MLSDK_API::MLAudioCreateInputFromMixedCaptureShim
+
+}
+
+#endif // !defined(WITH_MLSDK) || WITH_MLSDK
diff --git a/Engine/Plugins/Lumin/MLSDK/Source/MLSDK/Public/Lumin/CAPIShims/LuminAPICamera.h b/Engine/Plugins/Lumin/MLSDK/Source/MLSDK/Public/Lumin/CAPIShims/LuminAPICamera.h
new file mode 100644
--- /dev/null
+++ b/Engine/Plugins/Lumin/MLSDK/Source/MLSDK/Public/Lumin/CAPIShims/LuminAPICamera.h
@@ -0,0 +1,55 @@
+// Copyright 1998-2019 Epic Games, Inc. All Rights Reserved.
+
+#pragma once
+
+#if !defined(WITH_MLSDK) || WITH_MLSDK
+
+#include "Lumin/CAPIShims/LuminAPI.h"
+
+LUMIN_THIRD_PARTY_INCLUDES_START
+#include <ml_camera.h>
+LUMIN_THIRD_PARTY_INCLUDES_END
+
+namespace MLSDK_API
+{
+
+CREATE_FUNCTION_SHIM(ml_camera, MLResult, MLCameraConnect)
+#define MLCameraConnect ::MLSDK_API::MLCameraConnectShim
+CREATE_FUNCTION_SHIM(ml_camera, MLResult, MLCameraDisconnect)
+#define MLCameraDisconnect ::MLSDK_API::MLCameraDisconnectShim
+CREATE_FUNCTION_SHIM(ml_camera, MLResult, MLCameraPrepareCapture)
+#define MLCameraPrepareCapture ::MLSDK_API::MLCameraPrepareCaptureShim
+CREATE_FUNCTION_SHIM(ml_camera, MLResult, MLCameraSetDeviceStatusCallbacks)
+#define MLCameraSetDeviceStatusCallbacks ::MLSDK_API::MLCameraSetDeviceStatusCallbacksShim
+CREATE_FUNCTION_SHIM(ml_camera, MLResult, MLCameraSetCaptureCallbacks)
+#define MLCameraSetCaptureCallbacks ::MLSDK_API::MLCameraSetCaptureCallbacksShim
+CREATE_FUNCTION_SHIM(ml_camera, MLResult, MLCameraSetOutputFormat)
+#define MLCameraSetOutputFormat ::MLSDK_API::MLCameraSetOutputFormatShim
+CREATE_FUNCTION_SHIM(ml_camera, MLResult, MLCameraCaptureImage)
+#define MLCameraCaptureImage ::MLSDK_API::MLCameraCaptureImageShim
+CREATE_FUNCTION_SHIM(ml_camera, MLResult, MLCameraCaptureImageRaw)
+#define MLCameraCaptureImageRaw ::MLSDK_API::MLCameraCaptureImageRawShim
+CREATE_FUNCTION_SHIM(ml_camera, MLResult, MLCameraCaptureVideoStart)
+#define MLCameraCaptureVideoStart ::MLSDK_API::MLCameraCaptureVideoStartShim
+CREATE_FUNCTION_SHIM(ml_camera, MLResult, MLCameraCaptureVideoStop)
+#define MLCameraCaptureVideoStop ::MLSDK_API::MLCameraCaptureVideoStopShim
+CREATE_FUNCTION_SHIM(ml_camera, MLResult, MLCameraGetDeviceStatus)
+#define MLCameraGetDeviceStatus ::MLSDK_API::MLCameraGetDeviceStatusShim
+CREATE_FUNCTION_SHIM(ml_camera, MLResult, MLCameraGetCaptureStatus)
+#define MLCameraGetCaptureStatus ::MLSDK_API::MLCameraGetCaptureStatusShim
+CREATE_FUNCTION_SHIM(ml_camera, MLResult, MLCameraGetErrorCode)
+#define MLCameraGetErrorCode ::MLSDK_API::MLCameraGetErrorCodeShim
+CREATE_FUNCTION_SHIM(ml_camera, MLResult, MLCameraGetPreviewStream)
+#define MLCameraGetPreviewStream ::MLSDK_API::MLCameraGetPreviewStreamShim
+CREATE_FUNCTION_SHIM(ml_camera, MLResult, MLCameraGetImageStream)
+#define MLCameraGetImageStream ::MLSDK_API::MLCameraGetImageStreamShim
+CREATE_FUNCTION_SHIM(ml_camera, MLResult, MLCameraGetCaptureResultExtras)
+#define MLCameraGetCaptureResultExtras ::MLSDK_API::MLCameraGetCaptureResultExtrasShim
+CREATE_FUNCTION_SHIM(ml_camera, MLResult, MLCameraGetCameraCharacteristics)
+#define MLCameraGetCameraCharacteristics ::MLSDK_API::MLCameraGetCameraCharacteristicsShim
+CREATE_FUNCTION_SHIM(ml_camera, MLResult, MLCameraGetResultMetadata)
+#define MLCameraGetResultMetadata ::MLSDK_API::MLCameraGetResultMetadataShim
+
+}
+
+#endif // !defined(WITH_MLSDK) || WITH_MLSDK
diff --git a/Engine/Plugins/Lumin/MLSDK/Source/MLSDK/Public/Lumin/CAPIShims/LuminAPICameraMetadata.h b/Engine/Plugins/Lumin/MLSDK/Source/MLSDK/Public/Lumin/CAPIShims/LuminAPICameraMetadata.h
new file mode 100644
--- /dev/null
+++ b/Engine/Plugins/Lumin/MLSDK/Source/MLSDK/Public/Lumin/CAPIShims/LuminAPICameraMetadata.h
@@ -0,0 +1,140 @@
+// Copyright 1998-2019 Epic Games, Inc. All Rights Reserved.
+
+#pragma once
+
+#if !defined(WITH_MLSDK) || WITH_MLSDK
+
+#include "Lumin/CAPIShims/LuminAPI.h"
+
+LUMIN_THIRD_PARTY_INCLUDES_START
+#include <ml_camera_metadata.h>
+#include <ml_camera_metadata_tags.h>
+LUMIN_THIRD_PARTY_INCLUDES_END
+
+namespace MLSDK_API
+{
+
+CREATE_FUNCTION_SHIM(ml_camera_metadata, MLResult, MLCameraMetadataGetColorCorrectionAvailableAberrationModes)
+#define MLCameraMetadataGetColorCorrectionAvailableAberrationModes ::MLSDK_API::MLCameraMetadataGetColorCorrectionAvailableAberrationModesShim
+CREATE_FUNCTION_SHIM(ml_camera_metadata, MLResult, MLCameraMetadataGetControlAEAvailableModes)
+#define MLCameraMetadataGetControlAEAvailableModes ::MLSDK_API::MLCameraMetadataGetControlAEAvailableModesShim
+CREATE_FUNCTION_SHIM(ml_camera_metadata, MLResult, MLCameraMetadataGetControlAECompensationRange)
+#define MLCameraMetadataGetControlAECompensationRange ::MLSDK_API::MLCameraMetadataGetControlAECompensationRangeShim
+CREATE_FUNCTION_SHIM(ml_camera_metadata, MLResult, MLCameraMetadataGetControlAECompensationStep)
+#define MLCameraMetadataGetControlAECompensationStep ::MLSDK_API::MLCameraMetadataGetControlAECompensationStepShim
+CREATE_FUNCTION_SHIM(ml_camera_metadata, MLResult, MLCameraMetadataGetControlAELockAvailable)
+#define MLCameraMetadataGetControlAELockAvailable ::MLSDK_API::MLCameraMetadataGetControlAELockAvailableShim
+CREATE_FUNCTION_SHIM(ml_camera_metadata, MLResult, MLCameraMetadataGetControlAWBAvailableModes)
+#define MLCameraMetadataGetControlAWBAvailableModes ::MLSDK_API::MLCameraMetadataGetControlAWBAvailableModesShim
+CREATE_FUNCTION_SHIM(ml_camera_metadata, MLResult, MLCameraMetadataGetControlAWBLockAvailable)
+#define MLCameraMetadataGetControlAWBLockAvailable ::MLSDK_API::MLCameraMetadataGetControlAWBLockAvailableShim
+CREATE_FUNCTION_SHIM(ml_camera_metadata, MLResult, MLCameraMetadataGetScalerProcessedSizes)
+#define MLCameraMetadataGetScalerProcessedSizes ::MLSDK_API::MLCameraMetadataGetScalerProcessedSizesShim
+CREATE_FUNCTION_SHIM(ml_camera_metadata, MLResult, MLCameraMetadataGetScalerAvailableMaxDigitalZoom)
+#define MLCameraMetadataGetScalerAvailableMaxDigitalZoom ::MLSDK_API::MLCameraMetadataGetScalerAvailableMaxDigitalZoomShim
+CREATE_FUNCTION_SHIM(ml_camera_metadata, MLResult, MLCameraMetadataGetScalerAvailableStreamConfigurations)
+#define MLCameraMetadataGetScalerAvailableStreamConfigurations ::MLSDK_API::MLCameraMetadataGetScalerAvailableStreamConfigurationsShim
+CREATE_FUNCTION_SHIM(ml_camera_metadata, MLResult, MLCameraMetadataGetSensorInfoActiveArraySize)
+#define MLCameraMetadataGetSensorInfoActiveArraySize ::MLSDK_API::MLCameraMetadataGetSensorInfoActiveArraySizeShim
+CREATE_FUNCTION_SHIM(ml_camera_metadata, MLResult, MLCameraMetadataGetSensorInfoSensitivityRange)
+#define MLCameraMetadataGetSensorInfoSensitivityRange ::MLSDK_API::MLCameraMetadataGetSensorInfoSensitivityRangeShim
+CREATE_FUNCTION_SHIM(ml_camera_metadata, MLResult, MLCameraMetadataGetSensorInfoExposureTimeRange)
+#define MLCameraMetadataGetSensorInfoExposureTimeRange ::MLSDK_API::MLCameraMetadataGetSensorInfoExposureTimeRangeShim
+CREATE_FUNCTION_SHIM(ml_camera_metadata, MLResult, MLCameraMetadataGetSensorOrientation)
+#define MLCameraMetadataGetSensorOrientation ::MLSDK_API::MLCameraMetadataGetSensorOrientationShim
+CREATE_FUNCTION_SHIM(ml_camera_metadata, MLResult, MLCameraMetadataSetSensorInfoExposureTimeRange)
+#define MLCameraMetadataSetSensorInfoExposureTimeRange ::MLSDK_API::MLCameraMetadataSetSensorInfoExposureTimeRangeShim
+CREATE_FUNCTION_SHIM(ml_camera_metadata, MLResult, MLCameraMetadataGetColorCorrectionModeRequestMetadata)
+#define MLCameraMetadataGetColorCorrectionModeRequestMetadata ::MLSDK_API::MLCameraMetadataGetColorCorrectionModeRequestMetadataShim
+CREATE_FUNCTION_SHIM(ml_camera_metadata, MLResult, MLCameraMetadataGetColorCorrectionTransformRequestMetadata)
+#define MLCameraMetadataGetColorCorrectionTransformRequestMetadata ::MLSDK_API::MLCameraMetadataGetColorCorrectionTransformRequestMetadataShim
+CREATE_FUNCTION_SHIM(ml_camera_metadata, MLResult, MLCameraMetadataGetColorCorrectionGainsRequestMetadata)
+#define MLCameraMetadataGetColorCorrectionGainsRequestMetadata ::MLSDK_API::MLCameraMetadataGetColorCorrectionGainsRequestMetadataShim
+CREATE_FUNCTION_SHIM(ml_camera_metadata, MLResult, MLCameraMetadataGetColorCorrectionAberrationModeRequestMetadata)
+#define MLCameraMetadataGetColorCorrectionAberrationModeRequestMetadata ::MLSDK_API::MLCameraMetadataGetColorCorrectionAberrationModeRequestMetadataShim
+CREATE_FUNCTION_SHIM(ml_camera_metadata, MLResult, MLCameraMetadataGetControlAEAntibandingModeRequestMetadata)
+#define MLCameraMetadataGetControlAEAntibandingModeRequestMetadata ::MLSDK_API::MLCameraMetadataGetControlAEAntibandingModeRequestMetadataShim
+CREATE_FUNCTION_SHIM(ml_camera_metadata, MLResult, MLCameraMetadataGetControlAEExposureCompensationRequestMetadata)
+#define MLCameraMetadataGetControlAEExposureCompensationRequestMetadata ::MLSDK_API::MLCameraMetadataGetControlAEExposureCompensationRequestMetadataShim
+CREATE_FUNCTION_SHIM(ml_camera_metadata, MLResult, MLCameraMetadataGetControlAELockRequestMetadata)
+#define MLCameraMetadataGetControlAELockRequestMetadata ::MLSDK_API::MLCameraMetadataGetControlAELockRequestMetadataShim
+CREATE_FUNCTION_SHIM(ml_camera_metadata, MLResult, MLCameraMetadataGetControlAEModeRequestMetadata)
+#define MLCameraMetadataGetControlAEModeRequestMetadata ::MLSDK_API::MLCameraMetadataGetControlAEModeRequestMetadataShim
+CREATE_FUNCTION_SHIM(ml_camera_metadata, MLResult, MLCameraMetadataGetControlAETargetFPSRangeRequestMetadata)
+#define MLCameraMetadataGetControlAETargetFPSRangeRequestMetadata ::MLSDK_API::MLCameraMetadataGetControlAETargetFPSRangeRequestMetadataShim
+CREATE_FUNCTION_SHIM(ml_camera_metadata, MLResult, MLCameraMetadataGetControlAWBLockRequestMetadata)
+#define MLCameraMetadataGetControlAWBLockRequestMetadata ::MLSDK_API::MLCameraMetadataGetControlAWBLockRequestMetadataShim
+CREATE_FUNCTION_SHIM(ml_camera_metadata, MLResult, MLCameraMetadataGetControlAWBModeRequestMetadata)
+#define MLCameraMetadataGetControlAWBModeRequestMetadata ::MLSDK_API::MLCameraMetadataGetControlAWBModeRequestMetadataShim
+CREATE_FUNCTION_SHIM(ml_camera_metadata, MLResult, MLCameraMetadataGetSensorExposureTimeRequestMetadata)
+#define MLCameraMetadataGetSensorExposureTimeRequestMetadata ::MLSDK_API::MLCameraMetadataGetSensorExposureTimeRequestMetadataShim
+CREATE_FUNCTION_SHIM(ml_camera_metadata, MLResult, MLCameraMetadataGetSensorSensitivityRequestMetadata)
+#define MLCameraMetadataGetSensorSensitivityRequestMetadata ::MLSDK_API::MLCameraMetadataGetSensorSensitivityRequestMetadataShim
+CREATE_FUNCTION_SHIM(ml_camera_metadata, MLResult, MLCameraMetadataGetScalerCropRegionRequestMetadata)
+#define MLCameraMetadataGetScalerCropRegionRequestMetadata ::MLSDK_API::MLCameraMetadataGetScalerCropRegionRequestMetadataShim
+CREATE_FUNCTION_SHIM(ml_camera_metadata, MLResult, MLCameraMetadataSetColorCorrectionMode)
+#define MLCameraMetadataSetColorCorrectionMode ::MLSDK_API::MLCameraMetadataSetColorCorrectionModeShim
+CREATE_FUNCTION_SHIM(ml_camera_metadata, MLResult, MLCameraMetadataSetColorCorrectionTransform)
+#define MLCameraMetadataSetColorCorrectionTransform ::MLSDK_API::MLCameraMetadataSetColorCorrectionTransformShim
+CREATE_FUNCTION_SHIM(ml_camera_metadata, MLResult, MLCameraMetadataSetColorCorrectionGains)
+#define MLCameraMetadataSetColorCorrectionGains ::MLSDK_API::MLCameraMetadataSetColorCorrectionGainsShim
+CREATE_FUNCTION_SHIM(ml_camera_metadata, MLResult, MLCameraMetadataSetColorCorrectionAberrationMode)
+#define MLCameraMetadataSetColorCorrectionAberrationMode ::MLSDK_API::MLCameraMetadataSetColorCorrectionAberrationModeShim
+CREATE_FUNCTION_SHIM(ml_camera_metadata, MLResult, MLCameraMetadataSetControlAEAntibandingMode)
+#define MLCameraMetadataSetControlAEAntibandingMode ::MLSDK_API::MLCameraMetadataSetControlAEAntibandingModeShim
+CREATE_FUNCTION_SHIM(ml_camera_metadata, MLResult, MLCameraMetadataSetControlAEExposureCompensation)
+#define MLCameraMetadataSetControlAEExposureCompensation ::MLSDK_API::MLCameraMetadataSetControlAEExposureCompensationShim
+CREATE_FUNCTION_SHIM(ml_camera_metadata, MLResult, MLCameraMetadataSetControlAELock)
+#define MLCameraMetadataSetControlAELock ::MLSDK_API::MLCameraMetadataSetControlAELockShim
+CREATE_FUNCTION_SHIM(ml_camera_metadata, MLResult, MLCameraMetadataSetControlAEMode)
+#define MLCameraMetadataSetControlAEMode ::MLSDK_API::MLCameraMetadataSetControlAEModeShim
+CREATE_FUNCTION_SHIM(ml_camera_metadata, MLResult, MLCameraMetadataSetControlAETargetFPSRange)
+#define MLCameraMetadataSetControlAETargetFPSRange ::MLSDK_API::MLCameraMetadataSetControlAETargetFPSRangeShim
+CREATE_FUNCTION_SHIM(ml_camera_metadata, MLResult, MLCameraMetadataSetControlAWBLock)
+#define MLCameraMetadataSetControlAWBLock ::MLSDK_API::MLCameraMetadataSetControlAWBLockShim
+CREATE_FUNCTION_SHIM(ml_camera_metadata, MLResult, MLCameraMetadataSetControlAWBMode)
+#define MLCameraMetadataSetControlAWBMode ::MLSDK_API::MLCameraMetadataSetControlAWBModeShim
+CREATE_FUNCTION_SHIM(ml_camera_metadata, MLResult, MLCameraMetadataSetSensorExposureTime)
+#define MLCameraMetadataSetSensorExposureTime ::MLSDK_API::MLCameraMetadataSetSensorExposureTimeShim
+CREATE_FUNCTION_SHIM(ml_camera_metadata, MLResult, MLCameraMetadataSetSensorSensitivity)
+#define MLCameraMetadataSetSensorSensitivity ::MLSDK_API::MLCameraMetadataSetSensorSensitivityShim
+CREATE_FUNCTION_SHIM(ml_camera_metadata, MLResult, MLCameraMetadataSetScalerCropRegion)
+#define MLCameraMetadataSetScalerCropRegion ::MLSDK_API::MLCameraMetadataSetScalerCropRegionShim
+CREATE_FUNCTION_SHIM(ml_camera_metadata, MLResult, MLCameraMetadataGetColorCorrectionModeResultMetadata)
+#define MLCameraMetadataGetColorCorrectionModeResultMetadata ::MLSDK_API::MLCameraMetadataGetColorCorrectionModeResultMetadataShim
+CREATE_FUNCTION_SHIM(ml_camera_metadata, MLResult, MLCameraMetadataGetColorCorrectionTransformResultMetadata)
+#define MLCameraMetadataGetColorCorrectionTransformResultMetadata ::MLSDK_API::MLCameraMetadataGetColorCorrectionTransformResultMetadataShim
+CREATE_FUNCTION_SHIM(ml_camera_metadata, MLResult, MLCameraMetadataGetColorCorrectionAberrationModeResultMetadata)
+#define MLCameraMetadataGetColorCorrectionAberrationModeResultMetadata ::MLSDK_API::MLCameraMetadataGetColorCorrectionAberrationModeResultMetadataShim
+CREATE_FUNCTION_SHIM(ml_camera_metadata, MLResult, MLCameraMetadataGetColorCorrectionGainsResultMetadata)
+#define MLCameraMetadataGetColorCorrectionGainsResultMetadata ::MLSDK_API::MLCameraMetadataGetColorCorrectionGainsResultMetadataShim
+CREATE_FUNCTION_SHIM(ml_camera_metadata, MLResult, MLCameraMetadataGetControlAEAntibandingModeResultMetadata)
+#define MLCameraMetadataGetControlAEAntibandingModeResultMetadata ::MLSDK_API::MLCameraMetadataGetControlAEAntibandingModeResultMetadataShim
+CREATE_FUNCTION_SHIM(ml_camera_metadata, MLResult, MLCameraMetadataGetControlAEExposureCompensationResultMetadata)
+#define MLCameraMetadataGetControlAEExposureCompensationResultMetadata ::MLSDK_API::MLCameraMetadataGetControlAEExposureCompensationResultMetadataShim
+CREATE_FUNCTION_SHIM(ml_camera_metadata, MLResult, MLCameraMetadataGetControlAELockResultMetadata)
+#define MLCameraMetadataGetControlAELockResultMetadata ::MLSDK_API::MLCameraMetadataGetControlAELockResultMetadataShim
+CREATE_FUNCTION_SHIM(ml_camera_metadata, MLResult, MLCameraMetadataGetControlAEModeResultMetadata)
+#define MLCameraMetadataGetControlAEModeResultMetadata ::MLSDK_API::MLCameraMetadataGetControlAEModeResultMetadataShim
+CREATE_FUNCTION_SHIM(ml_camera_metadata, MLResult, MLCameraMetadataGetControlAETargetFPSRangeResultMetadata)
+#define MLCameraMetadataGetControlAETargetFPSRangeResultMetadata ::MLSDK_API::MLCameraMetadataGetControlAETargetFPSRangeResultMetadataShim
+CREATE_FUNCTION_SHIM(ml_camera_metadata, MLResult, MLCameraMetadataGetControlAEStateResultMetadata)
+#define MLCameraMetadataGetControlAEStateResultMetadata ::MLSDK_API::MLCameraMetadataGetControlAEStateResultMetadataShim
+CREATE_FUNCTION_SHIM(ml_camera_metadata, MLResult, MLCameraMetadataGetControlAWBLockResultMetadata)
+#define MLCameraMetadataGetControlAWBLockResultMetadata ::MLSDK_API::MLCameraMetadataGetControlAWBLockResultMetadataShim
+CREATE_FUNCTION_SHIM(ml_camera_metadata, MLResult, MLCameraMetadataGetControlAWBStateResultMetadata)
+#define MLCameraMetadataGetControlAWBStateResultMetadata ::MLSDK_API::MLCameraMetadataGetControlAWBStateResultMetadataShim
+CREATE_FUNCTION_SHIM(ml_camera_metadata, MLResult, MLCameraMetadataGetSensorExposureTimeResultMetadata)
+#define MLCameraMetadataGetSensorExposureTimeResultMetadata ::MLSDK_API::MLCameraMetadataGetSensorExposureTimeResultMetadataShim
+CREATE_FUNCTION_SHIM(ml_camera_metadata, MLResult, MLCameraMetadataGetSensorSensitivityResultMetadata)
+#define MLCameraMetadataGetSensorSensitivityResultMetadata ::MLSDK_API::MLCameraMetadataGetSensorSensitivityResultMetadataShim
+CREATE_FUNCTION_SHIM(ml_camera_metadata, MLResult, MLCameraMetadataGetSensorTimestampResultMetadata)
+#define MLCameraMetadataGetSensorTimestampResultMetadata ::MLSDK_API::MLCameraMetadataGetSensorTimestampResultMetadataShim
+CREATE_FUNCTION_SHIM(ml_camera_metadata, MLResult, MLCameraMetadataGetScalerCropRegionResultMetadata)
+#define MLCameraMetadataGetScalerCropRegionResultMetadata ::MLSDK_API::MLCameraMetadataGetScalerCropRegionResultMetadataShim
+CREATE_FUNCTION_SHIM(ml_camera_metadata, MLResult, MLCameraMetadataGetSensorFrameDurationResultMetadata)
+#define MLCameraMetadataGetSensorFrameDurationResultMetadata ::MLSDK_API::MLCameraMetadataGetSensorFrameDurationResultMetadataShim
+
+}
+
+#endif // !defined(WITH_MLSDK) || WITH_MLSDK
diff --git a/Engine/Plugins/Lumin/MLSDK/Source/MLSDK/Public/Lumin/CAPIShims/LuminAPICloud.h b/Engine/Plugins/Lumin/MLSDK/Source/MLSDK/Public/Lumin/CAPIShims/LuminAPICloud.h
new file mode 100644
--- /dev/null
+++ b/Engine/Plugins/Lumin/MLSDK/Source/MLSDK/Public/Lumin/CAPIShims/LuminAPICloud.h
@@ -0,0 +1,13 @@
+// Copyright 1998-2019 Epic Games, Inc. All Rights Reserved.
+
+#pragma once
+
+#if !defined(WITH_MLSDK) || WITH_MLSDK
+
+#include "Lumin/CAPIShims/LuminAPI.h"
+
+LUMIN_THIRD_PARTY_INCLUDES_START
+#include <ml_cloud.h>
+LUMIN_THIRD_PARTY_INCLUDES_END
+
+#endif // !defined(WITH_MLSDK) || WITH_MLSDK
diff --git a/Engine/Plugins/Lumin/MLSDK/Source/MLSDK/Public/Lumin/CAPIShims/LuminAPIContacts.h b/Engine/Plugins/Lumin/MLSDK/Source/MLSDK/Public/Lumin/CAPIShims/LuminAPIContacts.h
new file mode 100644
--- /dev/null
+++ b/Engine/Plugins/Lumin/MLSDK/Source/MLSDK/Public/Lumin/CAPIShims/LuminAPIContacts.h
@@ -0,0 +1,42 @@
+// Copyright 1998-2019 Epic Games, Inc. All Rights Reserved.
+
+#pragma once
+
+#if !defined(WITH_MLSDK) || WITH_MLSDK
+
+#include "Lumin/CAPIShims/LuminAPI.h"
+
+LUMIN_THIRD_PARTY_INCLUDES_START
+#include <ml_contacts.h>
+LUMIN_THIRD_PARTY_INCLUDES_END
+
+namespace MLSDK_API
+{
+
+CREATE_FUNCTION_SHIM(ml_contacts, MLResult, MLContactsStartup)
+#define MLContactsStartup ::MLSDK_API::MLContactsStartupShim
+CREATE_FUNCTION_SHIM(ml_contacts, MLResult, MLContactsShutdown)
+#define MLContactsShutdown ::MLSDK_API::MLContactsShutdownShim
+CREATE_FUNCTION_SHIM(ml_contacts, MLResult, MLContactsRequestInsert)
+#define MLContactsRequestInsert ::MLSDK_API::MLContactsRequestInsertShim
+CREATE_FUNCTION_SHIM(ml_contacts, MLResult, MLContactsRequestUpdate)
+#define MLContactsRequestUpdate ::MLSDK_API::MLContactsRequestUpdateShim
+CREATE_FUNCTION_SHIM(ml_contacts, MLResult, MLContactsRequestDelete)
+#define MLContactsRequestDelete ::MLSDK_API::MLContactsRequestDeleteShim
+CREATE_FUNCTION_SHIM(ml_contacts, MLResult, MLContactsTryGetOperationResult)
+#define MLContactsTryGetOperationResult ::MLSDK_API::MLContactsTryGetOperationResultShim
+CREATE_FUNCTION_SHIM(ml_contacts, MLResult, MLContactsRequestList)
+#define MLContactsRequestList ::MLSDK_API::MLContactsRequestListShim
+CREATE_FUNCTION_SHIM(ml_contacts, MLResult, MLContactsRequestSearch)
+#define MLContactsRequestSearch ::MLSDK_API::MLContactsRequestSearchShim
+CREATE_FUNCTION_SHIM(ml_contacts, MLResult, MLContactsRequestSelection)
+#define MLContactsRequestSelection ::MLSDK_API::MLContactsRequestSelectionShim
+CREATE_FUNCTION_SHIM(ml_contacts, MLResult, MLContactsTryGetListResult)
+#define MLContactsTryGetListResult ::MLSDK_API::MLContactsTryGetListResultShim
+CREATE_FUNCTION_SHIM(ml_contacts, MLResult, MLContactsReleaseRequestResources)
+#define MLContactsReleaseRequestResources ::MLSDK_API::MLContactsReleaseRequestResourcesShim
+CREATE_FUNCTION_SHIM(ml_contacts, const char *, MLContactsGetResultString)
+#define MLContactsGetResultString ::MLSDK_API::MLContactsGetResultStringShim
+}
+
+#endif // !defined(WITH_MLSDK) || WITH_MLSDK
diff --git a/Engine/Plugins/Lumin/MLSDK/Source/MLSDK/Public/Lumin/CAPIShims/LuminAPIController.h b/Engine/Plugins/Lumin/MLSDK/Source/MLSDK/Public/Lumin/CAPIShims/LuminAPIController.h
new file mode 100644
--- /dev/null
+++ b/Engine/Plugins/Lumin/MLSDK/Source/MLSDK/Public/Lumin/CAPIShims/LuminAPIController.h
@@ -0,0 +1,27 @@
+// Copyright 1998-2019 Epic Games, Inc. All Rights Reserved.
+
+#pragma once
+
+#if !defined(WITH_MLSDK) || WITH_MLSDK
+
+#include "Lumin/CAPIShims/LuminAPI.h"
+
+LUMIN_THIRD_PARTY_INCLUDES_START
+#include <ml_controller.h>
+LUMIN_THIRD_PARTY_INCLUDES_END
+
+namespace MLSDK_API
+{
+
+CREATE_FUNCTION_SHIM(ml_perception_client, MLResult, MLControllerCreate)
+#define MLControllerCreate ::MLSDK_API::MLControllerCreateShim
+CREATE_FUNCTION_SHIM(ml_perception_client, MLResult, MLControllerCreateEx)
+#define MLControllerCreateEx ::MLSDK_API::MLControllerCreateExShim
+CREATE_FUNCTION_SHIM(ml_perception_client, MLResult, MLControllerDestroy)
+#define MLControllerDestroy ::MLSDK_API::MLControllerDestroyShim
+CREATE_FUNCTION_SHIM(ml_perception_client, MLResult, MLControllerGetState)
+#define MLControllerGetState ::MLSDK_API::MLControllerGetStateShim
+
+}
+
+#endif // !defined(WITH_MLSDK) || WITH_MLSDK
diff --git a/Engine/Plugins/Lumin/MLSDK/Source/MLSDK/Public/Lumin/CAPIShims/LuminAPICoordinateFrameUID.h b/Engine/Plugins/Lumin/MLSDK/Source/MLSDK/Public/Lumin/CAPIShims/LuminAPICoordinateFrameUID.h
new file mode 100644
--- /dev/null
+++ b/Engine/Plugins/Lumin/MLSDK/Source/MLSDK/Public/Lumin/CAPIShims/LuminAPICoordinateFrameUID.h
@@ -0,0 +1,13 @@
+// Copyright 1998-2019 Epic Games, Inc. All Rights Reserved.
+
+#pragma once
+
+#if !defined(WITH_MLSDK) || WITH_MLSDK
+
+#include "Lumin/CAPIShims/LuminAPI.h"
+
+LUMIN_THIRD_PARTY_INCLUDES_START
+#include <ml_coordinate_frame_uid.h>
+LUMIN_THIRD_PARTY_INCLUDES_END
+
+#endif // !defined(WITH_MLSDK) || WITH_MLSDK
diff --git a/Engine/Plugins/Lumin/MLSDK/Source/MLSDK/Public/Lumin/CAPIShims/LuminAPIDataArray.h b/Engine/Plugins/Lumin/MLSDK/Source/MLSDK/Public/Lumin/CAPIShims/LuminAPIDataArray.h
new file mode 100644
--- /dev/null
+++ b/Engine/Plugins/Lumin/MLSDK/Source/MLSDK/Public/Lumin/CAPIShims/LuminAPIDataArray.h
@@ -0,0 +1,25 @@
+// Copyright 1998-2019 Epic Games, Inc. All Rights Reserved.
+
+#pragma once
+
+#if !defined(WITH_MLSDK) || WITH_MLSDK
+
+#include "Lumin/CAPIShims/LuminAPI.h"
+
+LUMIN_THIRD_PARTY_INCLUDES_START
+#include <ml_data_array.h>
+LUMIN_THIRD_PARTY_INCLUDES_END
+
+namespace MLSDK_API
+{
+
+CREATE_FUNCTION_SHIM(ml_perception_client, MLResult, MLDataArrayInitDiff)
+#define MLDataArrayInitDiff ::MLSDK_API::MLDataArrayInitDiffShim
+CREATE_FUNCTION_SHIM(ml_perception_client, MLResult, MLDataArrayTryLock)
+#define MLDataArrayTryLock ::MLSDK_API::MLDataArrayTryLockShim
+CREATE_FUNCTION_SHIM(ml_perception_client, MLResult, MLDataArrayUnlock)
+#define MLDataArrayUnlock ::MLSDK_API::MLDataArrayUnlockShim
+
+}
+
+#endif // !defined(WITH_MLSDK) || WITH_MLSDK
diff --git a/Engine/Plugins/Lumin/MLSDK/Source/MLSDK/Public/Lumin/CAPIShims/LuminAPIEyeTracking.h b/Engine/Plugins/Lumin/MLSDK/Source/MLSDK/Public/Lumin/CAPIShims/LuminAPIEyeTracking.h
new file mode 100644
--- /dev/null
+++ b/Engine/Plugins/Lumin/MLSDK/Source/MLSDK/Public/Lumin/CAPIShims/LuminAPIEyeTracking.h
@@ -0,0 +1,27 @@
+// Copyright 1998-2019 Epic Games, Inc. All Rights Reserved.
+
+#pragma once
+
+#if !defined(WITH_MLSDK) || WITH_MLSDK
+
+#include "Lumin/CAPIShims/LuminAPI.h"
+
+LUMIN_THIRD_PARTY_INCLUDES_START
+#include <ml_eye_tracking.h>
+LUMIN_THIRD_PARTY_INCLUDES_END
+
+namespace MLSDK_API
+{
+
+CREATE_FUNCTION_SHIM(ml_perception_client, MLResult, MLEyeTrackingCreate)
+#define MLEyeTrackingCreate ::MLSDK_API::MLEyeTrackingCreateShim
+CREATE_FUNCTION_SHIM(ml_perception_client, MLResult, MLEyeTrackingDestroy)
+#define MLEyeTrackingDestroy ::MLSDK_API::MLEyeTrackingDestroyShim
+CREATE_FUNCTION_SHIM(ml_perception_client, MLResult, MLEyeTrackingGetStaticData)
+#define MLEyeTrackingGetStaticData ::MLSDK_API::MLEyeTrackingGetStaticDataShim
+CREATE_FUNCTION_SHIM(ml_perception_client, MLResult, MLEyeTrackingGetState)
+#define MLEyeTrackingGetState ::MLSDK_API::MLEyeTrackingGetStateShim
+
+}
+
+#endif // !defined(WITH_MLSDK) || WITH_MLSDK
diff --git a/Engine/Plugins/Lumin/MLSDK/Source/MLSDK/Public/Lumin/CAPIShims/LuminAPIFoundObject.h b/Engine/Plugins/Lumin/MLSDK/Source/MLSDK/Public/Lumin/CAPIShims/LuminAPIFoundObject.h
new file mode 100644
--- /dev/null
+++ b/Engine/Plugins/Lumin/MLSDK/Source/MLSDK/Public/Lumin/CAPIShims/LuminAPIFoundObject.h
@@ -0,0 +1,69 @@
+// Copyright 1998-2019 Epic Games, Inc. All Rights Reserved.
+
+#pragma once
+
+#if !defined(WITH_MLSDK) || WITH_MLSDK
+
+#include "Lumin/CAPIShims/LuminAPI.h"
+
+LUMIN_THIRD_PARTY_INCLUDES_START
+#include <ml_found_object.h>
+LUMIN_THIRD_PARTY_INCLUDES_END
+
+namespace MLSDK_API
+{
+
+CREATE_FUNCTION_SHIM(ml_perception_client, MLResult, MLFoundObjectTrackerCreate)
+#define MLFoundObjectTrackerCreate ::MLSDK_API::MLFoundObjectTrackerCreateShim
+CREATE_FUNCTION_SHIM(ml_perception_client, MLResult, MLFoundObjectTrackerInitSettings)
+#define MLFoundObjectTrackerInitSettings ::MLSDK_API::MLFoundObjectTrackerInitSettingsShim
+CREATE_FUNCTION_SHIM(ml_perception_client, MLResult, MLFoundObjectTrackerUpdateSettings)
+#define MLFoundObjectTrackerUpdateSettings ::MLSDK_API::MLFoundObjectTrackerUpdateSettingsShim
+CREATE_FUNCTION_SHIM(ml_perception_client, MLResult, MLFoundObjectQuery)
+#define MLFoundObjectQuery ::MLSDK_API::MLFoundObjectQueryShim
+CREATE_FUNCTION_SHIM(ml_perception_client, MLResult, MLFoundObjectGetResultCount)
+#define MLFoundObjectGetResultCount ::MLSDK_API::MLFoundObjectGetResultCountShim
+CREATE_FUNCTION_SHIM(ml_perception_client, MLResult, MLFoundObjectGetResult)
+#define MLFoundObjectGetResult ::MLSDK_API::MLFoundObjectGetResultShim
+CREATE_FUNCTION_SHIM(ml_perception_client, MLResult, MLFoundObjectGetProperty)
+#define MLFoundObjectGetProperty ::MLSDK_API::MLFoundObjectGetPropertyShim
+CREATE_FUNCTION_SHIM(ml_perception_client, MLResult, MLFoundObjectRequestPropertyData)
+#define MLFoundObjectRequestPropertyData ::MLSDK_API::MLFoundObjectRequestPropertyDataShim
+CREATE_FUNCTION_SHIM(ml_perception_client, MLResult, MLFoundObjectGetPropertyData)
+#define MLFoundObjectGetPropertyData ::MLSDK_API::MLFoundObjectGetPropertyDataShim
+CREATE_FUNCTION_SHIM(ml_perception_client, MLResult, MLFoundObjectSetBinding)
+#define MLFoundObjectSetBinding ::MLSDK_API::MLFoundObjectSetBindingShim
+CREATE_FUNCTION_SHIM(ml_perception_client, MLResult, MLFoundObjectRemoveBinding)
+#define MLFoundObjectRemoveBinding ::MLSDK_API::MLFoundObjectRemoveBindingShim
+CREATE_FUNCTION_SHIM(ml_perception_client, MLResult, MLFoundObjectGetBindingCount)
+#define MLFoundObjectGetBindingCount ::MLSDK_API::MLFoundObjectGetBindingCountShim
+CREATE_FUNCTION_SHIM(ml_perception_client, MLResult, MLFoundObjectGetFoundObjectBindingCount)
+#define MLFoundObjectGetFoundObjectBindingCount ::MLSDK_API::MLFoundObjectGetFoundObjectBindingCountShim
+CREATE_FUNCTION_SHIM(ml_perception_client, MLResult, MLFoundObjectGetBinding)
+#define MLFoundObjectGetBinding ::MLSDK_API::MLFoundObjectGetBindingShim
+CREATE_FUNCTION_SHIM(ml_perception_client, MLResult, MLFoundObjectGetBindingByVirtualId)
+#define MLFoundObjectGetBindingByVirtualId ::MLSDK_API::MLFoundObjectGetBindingByVirtualIdShim
+CREATE_FUNCTION_SHIM(ml_perception_client, MLResult, MLFoundObjectGetAvailableLabelsCount)
+#define MLFoundObjectGetAvailableLabelsCount ::MLSDK_API::MLFoundObjectGetAvailableLabelsCountShim
+CREATE_FUNCTION_SHIM(ml_perception_client, MLResult, MLFoundObjectGetUniqueLabel)
+#define MLFoundObjectGetUniqueLabel ::MLSDK_API::MLFoundObjectGetUniqueLabelShim
+CREATE_FUNCTION_SHIM(ml_perception_client, MLResult, MLFoundObjectSubscribe)
+#define MLFoundObjectSubscribe ::MLSDK_API::MLFoundObjectSubscribeShim
+CREATE_FUNCTION_SHIM(ml_perception_client, MLResult, MLFoundObjectUnsubscribe)
+#define MLFoundObjectUnsubscribe ::MLSDK_API::MLFoundObjectUnsubscribeShim
+CREATE_FUNCTION_SHIM(ml_perception_client, MLResult, MLFoundObjectGetSubscriptionList)
+#define MLFoundObjectGetSubscriptionList ::MLSDK_API::MLFoundObjectGetSubscriptionListShim
+CREATE_FUNCTION_SHIM(ml_perception_client, MLResult, MLFoundObjectGetSubscriptionIdAtIndex)
+#define MLFoundObjectGetSubscriptionIdAtIndex ::MLSDK_API::MLFoundObjectGetSubscriptionIdAtIndexShim
+CREATE_FUNCTION_SHIM(ml_perception_client, MLResult, MLFoundObjectGetStateChangesCount)
+#define MLFoundObjectGetStateChangesCount ::MLSDK_API::MLFoundObjectGetStateChangesCountShim
+CREATE_FUNCTION_SHIM(ml_perception_client, MLResult, MLFoundObjectGetStateChangesCountForIndex)
+#define MLFoundObjectGetStateChangesCountForIndex ::MLSDK_API::MLFoundObjectGetStateChangesCountForIndexShim
+CREATE_FUNCTION_SHIM(ml_perception_client, MLResult, MLFoundObjectGetNextStateChange)
+#define MLFoundObjectGetNextStateChange ::MLSDK_API::MLFoundObjectGetNextStateChangeShim
+CREATE_FUNCTION_SHIM(ml_perception_client, MLResult, MLFoundObjectTrackerDestroy)
+#define MLFoundObjectTrackerDestroy ::MLSDK_API::MLFoundObjectTrackerDestroyShim
+
+}
+
+#endif // !defined(WITH_MLSDK) || WITH_MLSDK
diff --git a/Engine/Plugins/Lumin/MLSDK/Source/MLSDK/Public/Lumin/CAPIShims/LuminAPIGraphics.h b/Engine/Plugins/Lumin/MLSDK/Source/MLSDK/Public/Lumin/CAPIShims/LuminAPIGraphics.h
new file mode 100644
--- /dev/null
+++ b/Engine/Plugins/Lumin/MLSDK/Source/MLSDK/Public/Lumin/CAPIShims/LuminAPIGraphics.h
@@ -0,0 +1,62 @@
+// Copyright 1998-2019 Epic Games, Inc. All Rights Reserved.
+
+#pragma once
+
+#include "Lumin/CAPIShims/LuminAPI.h"
+
+#if !PLATFORM_MAC && defined(MLSDK_API_USE_VULKAN) && MLSDK_API_USE_VULKAN
+// Need all this for: class FVulkanTexture2D, and other Vulkan types.
+#include "VulkanRHIPrivate.h"
+#endif // !PLATFORM_MAC && defined(MLSDK_API_USE_VULKAN) && MLSDK_API_USE_VULKAN
+
+#if !defined(WITH_MLSDK) || WITH_MLSDK
+
+LUMIN_THIRD_PARTY_INCLUDES_START
+#include <ml_graphics.h>
+LUMIN_THIRD_PARTY_INCLUDES_END
+
+namespace MLSDK_API
+{
+
+CREATE_FUNCTION_SHIM(ml_graphics, MLResult, MLGraphicsEnableBlobCacheGL)
+#define MLGraphicsEnableBlobCacheGL ::MLSDK_API::MLGraphicsEnableBlobCacheGLShim
+CREATE_FUNCTION_SHIM(ml_graphics, MLResult, MLGraphicsCreateClientGL)
+#define MLGraphicsCreateClientGL ::MLSDK_API::MLGraphicsCreateClientGLShim
+CREATE_FUNCTION_SHIM(ml_graphics, uint32_t, MLGraphicsGLFormatFromMLSurfaceFormat)
+#define MLGraphicsGLFormatFromMLSurfaceFormat ::MLSDK_API::MLGraphicsGLFormatFromMLSurfaceFormatShim
+CREATE_FUNCTION_SHIM(ml_graphics, MLSurfaceFormat, MLGraphicsMLSurfaceFormatFromGLFormat)
+#define MLGraphicsMLSurfaceFormatFromGLFormat ::MLSDK_API::MLGraphicsMLSurfaceFormatFromGLFormatShim
+#ifdef VK_VERSION_1_0
+CREATE_FUNCTION_SHIM(ml_graphics, MLResult, MLGraphicsCreateClientVk)
+#define MLGraphicsCreateClientVk ::MLSDK_API::MLGraphicsCreateClientVkShim
+CREATE_FUNCTION_SHIM(ml_graphics, VkFormat, MLGraphicsVkFormatFromMLSurfaceFormat)
+#define MLGraphicsVkFormatFromMLSurfaceFormat ::MLSDK_API::MLGraphicsVkFormatFromMLSurfaceFormatShim
+CREATE_FUNCTION_SHIM(ml_graphics, MLSurfaceFormat, MLGraphicsMLSurfaceFormatFromVkFormat)
+#define MLGraphicsMLSurfaceFormatFromVkFormat ::MLSDK_API::MLGraphicsMLSurfaceFormatFromVkFormatShim
+#endif // VK_VERSION_1_0
+CREATE_FUNCTION_SHIM(ml_graphics, MLResult, MLGraphicsDestroyClient)
+#define MLGraphicsDestroyClient ::MLSDK_API::MLGraphicsDestroyClientShim
+CREATE_FUNCTION_SHIM(ml_graphics, MLResult, MLGraphicsSetFrameTimingHint)
+#define MLGraphicsSetFrameTimingHint ::MLSDK_API::MLGraphicsSetFrameTimingHintShim
+CREATE_FUNCTION_SHIM(ml_graphics, MLResult, MLGraphicsInitFrameParams)
+#define MLGraphicsInitFrameParams ::MLSDK_API::MLGraphicsInitFrameParamsShim
+CREATE_FUNCTION_SHIM(ml_graphics, MLResult, MLGraphicsBeginFrame)
+#define MLGraphicsBeginFrame ::MLSDK_API::MLGraphicsBeginFrameShim
+CREATE_FUNCTION_SHIM(ml_graphics, MLResult, MLGraphicsBeginFrameEx)
+#define MLGraphicsBeginFrameEx ::MLSDK_API::MLGraphicsBeginFrameExShim
+CREATE_FUNCTION_SHIM(ml_graphics, MLResult, MLGraphicsSignalSyncObjectGL)
+#define MLGraphicsSignalSyncObjectGL ::MLSDK_API::MLGraphicsSignalSyncObjectGLShim
+CREATE_FUNCTION_SHIM(ml_graphics, MLResult, MLGraphicsGetClipExtents)
+#define MLGraphicsGetClipExtents ::MLSDK_API::MLGraphicsGetClipExtentsShim
+CREATE_FUNCTION_SHIM(ml_graphics, MLResult, MLGraphicsGetClipExtentsEx)
+#define MLGraphicsGetClipExtentsEx ::MLSDK_API::MLGraphicsGetClipExtentsExShim
+CREATE_FUNCTION_SHIM(ml_graphics, MLResult, MLGraphicsGetRenderTargets)
+#define MLGraphicsGetRenderTargets ::MLSDK_API::MLGraphicsGetRenderTargetsShim
+CREATE_FUNCTION_SHIM(ml_graphics, MLResult, MLGraphicsEndFrame)
+#define MLGraphicsEndFrame ::MLSDK_API::MLGraphicsEndFrameShim
+CREATE_FUNCTION_SHIM(ml_graphics, MLResult, MLGraphicsGetClientPerformanceInfo)
+#define MLGraphicsGetClientPerformanceInfo ::MLSDK_API::MLGraphicsGetClientPerformanceInfoShim
+
+}
+
+#endif // !defined(WITH_MLSDK) || WITH_MLSDK
diff --git a/Engine/Plugins/Lumin/MLSDK/Source/MLSDK/Public/Lumin/CAPIShims/LuminAPIGraphicsUtils.h b/Engine/Plugins/Lumin/MLSDK/Source/MLSDK/Public/Lumin/CAPIShims/LuminAPIGraphicsUtils.h
new file mode 100644
--- /dev/null
+++ b/Engine/Plugins/Lumin/MLSDK/Source/MLSDK/Public/Lumin/CAPIShims/LuminAPIGraphicsUtils.h
@@ -0,0 +1,28 @@
+// Copyright 1998-2019 Epic Games, Inc. All Rights Reserved.
+
+#pragma once
+
+#if !defined(WITH_MLSDK) || WITH_MLSDK
+
+#include "Lumin/CAPIShims/LuminAPI.h"
+#include "LuminAPIGraphics.h"
+
+LUMIN_THIRD_PARTY_INCLUDES_START
+#include <ml_graphics_utils.h>
+LUMIN_THIRD_PARTY_INCLUDES_END
+
+namespace MLSDK_API
+{
+
+#ifdef VK_VERSION_1_0
+CREATE_FUNCTION_SHIM(ml_graphics_utils, MLResult, MLGraphicsEnumerateRequiredVkDeviceExtensionsForMediaHandleImport)
+#define MLGraphicsEnumerateRequiredVkDeviceExtensionsForMediaHandleImport ::MLSDK_API::MLGraphicsEnumerateRequiredVkDeviceExtensionsForMediaHandleImportShim
+CREATE_FUNCTION_SHIM(ml_graphics_utils, MLResult, MLGraphicsImportVkImageFromMediaHandle)
+#define MLGraphicsImportVkImageFromMediaHandle ::MLSDK_API::MLGraphicsImportVkImageFromMediaHandleShim
+#endif // VK_VERSION_1_0
+CREATE_FUNCTION_SHIM(ml_graphics_utils, MLResult, MLMeshingPopulateDepth)
+#define MLMeshingPopulateDepth ::MLSDK_API::MLMeshingPopulateDepthShim
+
+}
+
+#endif // !defined(WITH_MLSDK) || WITH_MLSDK
diff --git a/Engine/Plugins/Lumin/MLSDK/Source/MLSDK/Public/Lumin/CAPIShims/LuminAPIHandMeshing.h b/Engine/Plugins/Lumin/MLSDK/Source/MLSDK/Public/Lumin/CAPIShims/LuminAPIHandMeshing.h
new file mode 100644
--- /dev/null
+++ b/Engine/Plugins/Lumin/MLSDK/Source/MLSDK/Public/Lumin/CAPIShims/LuminAPIHandMeshing.h
@@ -0,0 +1,28 @@
+// Copyright 1998-2019 Epic Games, Inc. All Rights Reserved.
+#pragma once
+
+#if !defined(WITH_MLSDK) || WITH_MLSDK
+
+#include "Lumin/CAPIShims/LuminAPI.h"
+
+LUMIN_THIRD_PARTY_INCLUDES_START
+#include <ml_hand_meshing.h>
+LUMIN_THIRD_PARTY_INCLUDES_END
+
+namespace MLSDK_API
+{
+
+CREATE_FUNCTION_SHIM(ml_perception_client, MLResult, MLHandMeshingCreateClient)
+#define MLHandMeshingCreateClient ::MLSDK_API::MLHandMeshingCreateClientShim
+CREATE_FUNCTION_SHIM(ml_perception_client, MLResult, MLHandMeshingDestroyClient)
+#define MLHandMeshingDestroyClient ::MLSDK_API::MLHandMeshingDestroyClientShim
+CREATE_FUNCTION_SHIM(ml_perception_client, MLResult, MLHandMeshingRequestMesh)
+#define MLHandMeshingRequestMesh ::MLSDK_API::MLHandMeshingRequestMeshShim
+CREATE_FUNCTION_SHIM(ml_perception_client, MLResult, MLHandMeshingGetResult)
+#define MLHandMeshingGetResult ::MLSDK_API::MLHandMeshingGetResultShim
+CREATE_FUNCTION_SHIM(ml_perception_client, MLResult, MLHandMeshingFreeResource)
+#define MLHandMeshingFreeResource ::MLSDK_API::MLHandMeshingFreeResourceShim
+
+}
+
+#endif // !defined(WITH_MLSDK) || WITH_MLSDK
diff --git a/Engine/Plugins/Lumin/MLSDK/Source/MLSDK/Public/Lumin/CAPIShims/LuminAPIHandTracking.h b/Engine/Plugins/Lumin/MLSDK/Source/MLSDK/Public/Lumin/CAPIShims/LuminAPIHandTracking.h
new file mode 100644
--- /dev/null
+++ b/Engine/Plugins/Lumin/MLSDK/Source/MLSDK/Public/Lumin/CAPIShims/LuminAPIHandTracking.h
@@ -0,0 +1,31 @@
+// Copyright 1998-2019 Epic Games, Inc. All Rights Reserved.
+
+#pragma once
+
+#if !defined(WITH_MLSDK) || WITH_MLSDK
+
+#include "Lumin/CAPIShims/LuminAPI.h"
+
+LUMIN_THIRD_PARTY_INCLUDES_START
+#include <ml_hand_tracking.h>
+LUMIN_THIRD_PARTY_INCLUDES_END
+
+namespace MLSDK_API
+{
+
+CREATE_FUNCTION_SHIM(ml_perception_client, MLResult, MLHandTrackingCreate)
+#define MLHandTrackingCreate ::MLSDK_API::MLHandTrackingCreateShim
+CREATE_FUNCTION_SHIM(ml_perception_client, MLResult, MLHandTrackingDestroy)
+#define MLHandTrackingDestroy ::MLSDK_API::MLHandTrackingDestroyShim
+CREATE_FUNCTION_SHIM(ml_perception_client, MLResult, MLHandTrackingGetData)
+#define MLHandTrackingGetData ::MLSDK_API::MLHandTrackingGetDataShim
+CREATE_FUNCTION_SHIM(ml_perception_client, MLResult, MLHandTrackingGetStaticData)
+#define MLHandTrackingGetStaticData ::MLSDK_API::MLHandTrackingGetStaticDataShim
+CREATE_FUNCTION_SHIM(ml_perception_client, MLResult, MLHandTrackingSetConfiguration)
+#define MLHandTrackingSetConfiguration ::MLSDK_API::MLHandTrackingSetConfigurationShim
+CREATE_FUNCTION_SHIM(ml_perception_client, MLResult, MLHandTrackingGetConfiguration)
+#define MLHandTrackingGetConfiguration ::MLSDK_API::MLHandTrackingGetConfigurationShim
+
+}
+
+#endif // !defined(WITH_MLSDK) || WITH_MLSDK
diff --git a/Engine/Plugins/Lumin/MLSDK/Source/MLSDK/Public/Lumin/CAPIShims/LuminAPIHeadTracking.h b/Engine/Plugins/Lumin/MLSDK/Source/MLSDK/Public/Lumin/CAPIShims/LuminAPIHeadTracking.h
new file mode 100644
--- /dev/null
+++ b/Engine/Plugins/Lumin/MLSDK/Source/MLSDK/Public/Lumin/CAPIShims/LuminAPIHeadTracking.h
@@ -0,0 +1,29 @@
+// Copyright 1998-2019 Epic Games, Inc. All Rights Reserved.
+
+#pragma once
+
+#if !defined(WITH_MLSDK) || WITH_MLSDK
+
+#include "Lumin/CAPIShims/LuminAPI.h"
+
+LUMIN_THIRD_PARTY_INCLUDES_START
+#include <ml_head_tracking.h>
+LUMIN_THIRD_PARTY_INCLUDES_END
+
+namespace MLSDK_API
+{
+
+CREATE_FUNCTION_SHIM(ml_perception_client, MLResult, MLHeadTrackingCreate)
+#define MLHeadTrackingCreate ::MLSDK_API::MLHeadTrackingCreateShim
+CREATE_FUNCTION_SHIM(ml_perception_client, MLResult, MLHeadTrackingDestroy)
+#define MLHeadTrackingDestroy ::MLSDK_API::MLHeadTrackingDestroyShim
+CREATE_FUNCTION_SHIM(ml_perception_client, MLResult, MLHeadTrackingGetStaticData)
+#define MLHeadTrackingGetStaticData ::MLSDK_API::MLHeadTrackingGetStaticDataShim
+CREATE_FUNCTION_SHIM(ml_perception_client, MLResult, MLHeadTrackingGetState)
+#define MLHeadTrackingGetState ::MLSDK_API::MLHeadTrackingGetStateShim
+CREATE_FUNCTION_SHIM(ml_perception_client, MLResult, MLHeadTrackingGetMapEvents)
+#define MLHeadTrackingGetMapEvents ::MLSDK_API::MLHeadTrackingGetMapEventsShim
+
+}
+
+#endif // !defined(WITH_MLSDK) || WITH_MLSDK
diff --git a/Engine/Plugins/Lumin/MLSDK/Source/MLSDK/Public/Lumin/CAPIShims/LuminAPIIdentity.h b/Engine/Plugins/Lumin/MLSDK/Source/MLSDK/Public/Lumin/CAPIShims/LuminAPIIdentity.h
new file mode 100644
--- /dev/null
+++ b/Engine/Plugins/Lumin/MLSDK/Source/MLSDK/Public/Lumin/CAPIShims/LuminAPIIdentity.h
@@ -0,0 +1,37 @@
+// Copyright 1998-2019 Epic Games, Inc. All Rights Reserved.
+
+#pragma once
+
+#if !defined(WITH_MLSDK) || WITH_MLSDK
+
+#include "Lumin/CAPIShims/LuminAPI.h"
+
+LUMIN_THIRD_PARTY_INCLUDES_START
+#include <ml_identity.h>
+LUMIN_THIRD_PARTY_INCLUDES_END
+
+namespace MLSDK_API
+{
+
+CREATE_FUNCTION_SHIM(ml_identity, MLResult, MLIdentityGetAttributeNames)
+#define MLIdentityGetAttributeNames ::MLSDK_API::MLIdentityGetAttributeNamesShim
+CREATE_FUNCTION_SHIM(ml_identity, MLResult, MLIdentityGetKnownAttributeNames)
+#define MLIdentityGetKnownAttributeNames ::MLSDK_API::MLIdentityGetKnownAttributeNamesShim
+CREATE_FUNCTION_SHIM(ml_identity, MLResult, MLIdentityGetAttributeNamesAsync)
+#define MLIdentityGetAttributeNamesAsync ::MLSDK_API::MLIdentityGetAttributeNamesAsyncShim
+CREATE_FUNCTION_SHIM(ml_identity, MLResult, MLIdentityGetAttributeNamesWait)
+#define MLIdentityGetAttributeNamesWait ::MLSDK_API::MLIdentityGetAttributeNamesWaitShim
+CREATE_FUNCTION_SHIM(ml_identity, MLResult, MLIdentityRequestAttributeValues)
+#define MLIdentityRequestAttributeValues ::MLSDK_API::MLIdentityRequestAttributeValuesShim
+CREATE_FUNCTION_SHIM(ml_identity, MLResult, MLIdentityRequestAttributeValuesAsync)
+#define MLIdentityRequestAttributeValuesAsync ::MLSDK_API::MLIdentityRequestAttributeValuesAsyncShim
+CREATE_FUNCTION_SHIM(ml_identity, MLResult, MLIdentityRequestAttributeValuesWait)
+#define MLIdentityRequestAttributeValuesWait ::MLSDK_API::MLIdentityRequestAttributeValuesWaitShim
+CREATE_FUNCTION_SHIM(ml_identity, MLResult, MLIdentityReleaseUserProfile)
+#define MLIdentityReleaseUserProfile ::MLSDK_API::MLIdentityReleaseUserProfileShim
+CREATE_FUNCTION_SHIM(ml_identity, const char*, MLIdentityGetResultString)
+#define MLIdentityGetResultString ::MLSDK_API::MLIdentityGetResultStringShim
+
+}
+
+#endif // !defined(WITH_MLSDK) || WITH_MLSDK
diff --git a/Engine/Plugins/Lumin/MLSDK/Source/MLSDK/Public/Lumin/CAPIShims/LuminAPIImage.h b/Engine/Plugins/Lumin/MLSDK/Source/MLSDK/Public/Lumin/CAPIShims/LuminAPIImage.h
new file mode 100644
--- /dev/null
+++ b/Engine/Plugins/Lumin/MLSDK/Source/MLSDK/Public/Lumin/CAPIShims/LuminAPIImage.h
@@ -0,0 +1,13 @@
+// Copyright 1998-2019 Epic Games, Inc. All Rights Reserved.
+
+#pragma once
+
+#if !defined(WITH_MLSDK) || WITH_MLSDK
+
+#include "Lumin/CAPIShims/LuminAPI.h"
+
+LUMIN_THIRD_PARTY_INCLUDES_START
+#include <ml_image.h>
+LUMIN_THIRD_PARTY_INCLUDES_END
+
+#endif // !defined(WITH_MLSDK) || WITH_MLSDK
diff --git a/Engine/Plugins/Lumin/MLSDK/Source/MLSDK/Public/Lumin/CAPIShims/LuminAPIImageTracking.h b/Engine/Plugins/Lumin/MLSDK/Source/MLSDK/Public/Lumin/CAPIShims/LuminAPIImageTracking.h
new file mode 100644
--- /dev/null
+++ b/Engine/Plugins/Lumin/MLSDK/Source/MLSDK/Public/Lumin/CAPIShims/LuminAPIImageTracking.h
@@ -0,0 +1,39 @@
+// Copyright 1998-2019 Epic Games, Inc. All Rights Reserved.
+
+#pragma once
+
+#if !defined(WITH_MLSDK) || WITH_MLSDK
+
+#include "Lumin/CAPIShims/LuminAPI.h"
+
+LUMIN_THIRD_PARTY_INCLUDES_START
+#include <ml_image_tracking.h>
+LUMIN_THIRD_PARTY_INCLUDES_END
+
+namespace MLSDK_API
+{
+
+CREATE_FUNCTION_SHIM(ml_perception_client, MLResult, MLImageTrackerCreate)
+#define MLImageTrackerCreate ::MLSDK_API::MLImageTrackerCreateShim
+CREATE_FUNCTION_SHIM(ml_perception_client, MLResult, MLImageTrackerUpdateSettings)
+#define MLImageTrackerUpdateSettings ::MLSDK_API::MLImageTrackerUpdateSettingsShim
+CREATE_FUNCTION_SHIM(ml_perception_client, MLResult, MLImageTrackerDestroy)
+#define MLImageTrackerDestroy ::MLSDK_API::MLImageTrackerDestroyShim
+CREATE_FUNCTION_SHIM(ml_perception_client, MLResult, MLImageTrackerInitSettings)
+#define MLImageTrackerInitSettings ::MLSDK_API::MLImageTrackerInitSettingsShim
+CREATE_FUNCTION_SHIM(ml_perception_client, MLResult, MLImageTrackerAddTargetFromImageFile)
+#define MLImageTrackerAddTargetFromImageFile ::MLSDK_API::MLImageTrackerAddTargetFromImageFileShim
+CREATE_FUNCTION_SHIM(ml_perception_client, MLResult, MLImageTrackerAddTargetFromArray)
+#define MLImageTrackerAddTargetFromArray ::MLSDK_API::MLImageTrackerAddTargetFromArrayShim
+CREATE_FUNCTION_SHIM(ml_perception_client, MLResult, MLImageTrackerRemoveTarget)
+#define MLImageTrackerRemoveTarget ::MLSDK_API::MLImageTrackerRemoveTargetShim
+CREATE_FUNCTION_SHIM(ml_perception_client, MLResult, MLImageTrackerUpdateTargetSettings)
+#define MLImageTrackerUpdateTargetSettings ::MLSDK_API::MLImageTrackerUpdateTargetSettingsShim
+CREATE_FUNCTION_SHIM(ml_perception_client, MLResult, MLImageTrackerGetTargetStaticData)
+#define MLImageTrackerGetTargetStaticData ::MLSDK_API::MLImageTrackerGetTargetStaticDataShim
+CREATE_FUNCTION_SHIM(ml_perception_client, MLResult, MLImageTrackerGetTargetResult)
+#define MLImageTrackerGetTargetResult ::MLSDK_API::MLImageTrackerGetTargetResultShim
+
+}
+
+#endif // !defined(WITH_MLSDK) || WITH_MLSDK
diff --git a/Engine/Plugins/Lumin/MLSDK/Source/MLSDK/Public/Lumin/CAPIShims/LuminAPIInvokeAsync.h b/Engine/Plugins/Lumin/MLSDK/Source/MLSDK/Public/Lumin/CAPIShims/LuminAPIInvokeAsync.h
new file mode 100644
--- /dev/null
+++ b/Engine/Plugins/Lumin/MLSDK/Source/MLSDK/Public/Lumin/CAPIShims/LuminAPIInvokeAsync.h
@@ -0,0 +1,13 @@
+// Copyright 1998-2019 Epic Games, Inc. All Rights Reserved.
+
+#pragma once
+
+#if !defined(WITH_MLSDK) || WITH_MLSDK
+
+#include "Lumin/CAPIShims/LuminAPI.h"
+
+LUMIN_THIRD_PARTY_INCLUDES_START
+#include <ml_invoke_async.h>
+LUMIN_THIRD_PARTY_INCLUDES_END
+
+#endif // !defined(WITH_MLSDK) || WITH_MLSDK
diff --git a/Engine/Plugins/Lumin/MLSDK/Source/MLSDK/Public/Lumin/CAPIShims/LuminAPILightingTracking.h b/Engine/Plugins/Lumin/MLSDK/Source/MLSDK/Public/Lumin/CAPIShims/LuminAPILightingTracking.h
new file mode 100644
--- /dev/null
+++ b/Engine/Plugins/Lumin/MLSDK/Source/MLSDK/Public/Lumin/CAPIShims/LuminAPILightingTracking.h
@@ -0,0 +1,29 @@
+// Copyright 1998-2019 Epic Games, Inc. All Rights Reserved.
+
+#pragma once
+
+#if !defined(WITH_MLSDK) || WITH_MLSDK
+
+#include "Lumin/CAPIShims/LuminAPI.h"
+
+LUMIN_THIRD_PARTY_INCLUDES_START
+#include <ml_lighting_tracking.h>
+LUMIN_THIRD_PARTY_INCLUDES_END
+
+namespace MLSDK_API
+{
+
+CREATE_FUNCTION_SHIM(ml_perception_client, MLResult, MLLightingTrackingCreate)
+#define MLLightingTrackingCreate ::MLSDK_API::MLLightingTrackingCreateShim
+CREATE_FUNCTION_SHIM(ml_perception_client, MLResult, MLLightingTrackingDestroy)
+#define MLLightingTrackingDestroy ::MLSDK_API::MLLightingTrackingDestroyShim
+CREATE_FUNCTION_SHIM(ml_perception_client, MLResult, MLLightingTrackingGetAmbientGlobalState)
+#define MLLightingTrackingGetAmbientGlobalState ::MLSDK_API::MLLightingTrackingGetAmbientGlobalStateShim
+CREATE_FUNCTION_SHIM(ml_perception_client, MLResult, MLLightingTrackingGetAmbientGridState)
+#define MLLightingTrackingGetAmbientGridState ::MLSDK_API::MLLightingTrackingGetAmbientGridStateShim
+CREATE_FUNCTION_SHIM(ml_perception_client, MLResult, MLLightingTrackingGetColorTemperatureState)
+#define MLLightingTrackingGetColorTemperatureState ::MLSDK_API::MLLightingTrackingGetColorTemperatureStateShim
+
+}
+
+#endif // !defined(WITH_MLSDK) || WITH_MLSDK
diff --git a/Engine/Plugins/Lumin/MLSDK/Source/MLSDK/Public/Lumin/CAPIShims/LuminAPILocation.h b/Engine/Plugins/Lumin/MLSDK/Source/MLSDK/Public/Lumin/CAPIShims/LuminAPILocation.h
new file mode 100644
--- /dev/null
+++ b/Engine/Plugins/Lumin/MLSDK/Source/MLSDK/Public/Lumin/CAPIShims/LuminAPILocation.h
@@ -0,0 +1,25 @@
+// Copyright 1998-2019 Epic Games, Inc. All Rights Reserved.
+
+#pragma once
+
+#if !defined(WITH_MLSDK) || WITH_MLSDK
+
+#include "Lumin/CAPIShims/LuminAPI.h"
+
+LUMIN_THIRD_PARTY_INCLUDES_START
+#include <ml_location.h>
+LUMIN_THIRD_PARTY_INCLUDES_END
+
+namespace MLSDK_API
+{
+
+CREATE_FUNCTION_SHIM(ml_location, MLResult, MLLocationGetLastCoarseLocation)
+#define MLLocationGetLastCoarseLocation ::MLSDK_API::MLLocationGetLastCoarseLocationShim
+CREATE_FUNCTION_SHIM(ml_location, MLResult, MLLocationGetLastFineLocation)
+#define MLLocationGetLastFineLocation ::MLSDK_API::MLLocationGetLastFineLocationShim
+CREATE_FUNCTION_SHIM(ml_location, const char *, MLLocationGetResultString)
+#define MLLocationGetResultString ::MLSDK_API::MLLocationGetResultStringShim
+
+}
+
+#endif // !defined(WITH_MLSDK) || WITH_MLSDK
diff --git a/Engine/Plugins/Lumin/MLSDK/Source/MLSDK/Public/Lumin/CAPIShims/LuminAPIMediaCCParser.h b/Engine/Plugins/Lumin/MLSDK/Source/MLSDK/Public/Lumin/CAPIShims/LuminAPIMediaCCParser.h
new file mode 100644
--- /dev/null
+++ b/Engine/Plugins/Lumin/MLSDK/Source/MLSDK/Public/Lumin/CAPIShims/LuminAPIMediaCCParser.h
@@ -0,0 +1,28 @@
+// Copyright 1998-2019 Epic Games, Inc. All Rights Reserved.
+
+#pragma once
+
+#if !defined(WITH_MLSDK) || WITH_MLSDK
+
+#include "Lumin/CAPIShims/LuminAPI.h"
+
+LUMIN_THIRD_PARTY_INCLUDES_START
+#include <ml_media_common.h>
+#include <ml_media_cea608_caption.h>
+LUMIN_THIRD_PARTY_INCLUDES_END
+
+namespace MLSDK_API
+{
+
+CREATE_FUNCTION_SHIM(ml_media_ccparser, MLResult, MLMediaCCParserCreate)
+#define MLMediaCCParserCreate ::MLSDK_API::MLMediaCCParserCreateShim
+CREATE_FUNCTION_SHIM(ml_media_ccparser, MLResult, MLMediaCCParserGetDisplayable)
+#define MLMediaCCParserGetDisplayable ::MLSDK_API::MLMediaCCParserGetDisplayableShim
+CREATE_FUNCTION_SHIM(ml_media_ccparser, MLResult, MLMediaCCParserReleaseSegment)
+#define MLMediaCCParserReleaseSegment ::MLSDK_API::MLMediaCCParserReleaseSegmentShim
+CREATE_FUNCTION_SHIM(ml_media_ccparser, MLResult, MLMediaCCParserDestroy)
+#define MLMediaCCParserDestroy ::MLSDK_API::MLMediaCCParserDestroyShim
+
+}
+
+#endif // !defined(WITH_MLSDK) || WITH_MLSDK
diff --git a/Engine/Plugins/Lumin/MLSDK/Source/MLSDK/Public/Lumin/CAPIShims/LuminAPIMediaCodec.h b/Engine/Plugins/Lumin/MLSDK/Source/MLSDK/Public/Lumin/CAPIShims/LuminAPIMediaCodec.h
new file mode 100644
--- /dev/null
+++ b/Engine/Plugins/Lumin/MLSDK/Source/MLSDK/Public/Lumin/CAPIShims/LuminAPIMediaCodec.h
@@ -0,0 +1,70 @@
+// Copyright 1998-2019 Epic Games, Inc. All Rights Reserved.
+
+#pragma once
+
+#if !defined(WITH_MLSDK) || WITH_MLSDK
+
+#include "Lumin/CAPIShims/LuminAPI.h"
+
+LUMIN_THIRD_PARTY_INCLUDES_START
+#include <ml_media_common.h>
+#include <ml_media_codec.h>
+LUMIN_THIRD_PARTY_INCLUDES_END
+
+namespace MLSDK_API
+{
+
+CREATE_FUNCTION_SHIM(ml_mediacodec, MLResult, MLMediaCodecCreateCodec)
+#define MLMediaCodecCreateCodec ::MLSDK_API::MLMediaCodecCreateCodecShim
+CREATE_FUNCTION_SHIM(ml_mediacodec, MLResult, MLMediaCodecDestroy)
+#define MLMediaCodecDestroy ::MLSDK_API::MLMediaCodecDestroyShim
+CREATE_FUNCTION_SHIM(ml_mediacodec, MLResult, MLMediaCodecSetCallbacks)
+#define MLMediaCodecSetCallbacks ::MLSDK_API::MLMediaCodecSetCallbacksShim
+CREATE_FUNCTION_SHIM(ml_mediacodec, MLResult, MLMediaCodecGetName)
+#define MLMediaCodecGetName ::MLSDK_API::MLMediaCodecGetNameShim
+CREATE_FUNCTION_SHIM(ml_mediacodec, MLResult, MLMediaCodecSetSurfaceHint)
+#define MLMediaCodecSetSurfaceHint ::MLSDK_API::MLMediaCodecSetSurfaceHintShim
+CREATE_FUNCTION_SHIM(ml_mediacodec, MLResult, MLMediaCodecConfigure)
+#define MLMediaCodecConfigure ::MLSDK_API::MLMediaCodecConfigureShim
+CREATE_FUNCTION_SHIM(ml_mediacodec, MLResult, MLMediaCodecStart)
+#define MLMediaCodecStart ::MLSDK_API::MLMediaCodecStartShim
+CREATE_FUNCTION_SHIM(ml_mediacodec, MLResult, MLMediaCodecStop)
+#define MLMediaCodecStop ::MLSDK_API::MLMediaCodecStopShim
+CREATE_FUNCTION_SHIM(ml_mediacodec, MLResult, MLMediaCodecFlush)
+#define MLMediaCodecFlush ::MLSDK_API::MLMediaCodecFlushShim
+CREATE_FUNCTION_SHIM(ml_mediacodec, MLResult, MLMediaCodecGetInputBufferPointer)
+#define MLMediaCodecGetInputBufferPointer ::MLSDK_API::MLMediaCodecGetInputBufferPointerShim
+CREATE_FUNCTION_SHIM(ml_mediacodec, MLResult, MLMediaCodecGetOutputBufferPointer)
+#define MLMediaCodecGetOutputBufferPointer ::MLSDK_API::MLMediaCodecGetOutputBufferPointerShim
+CREATE_FUNCTION_SHIM(ml_mediacodec, MLResult, MLMediaCodecQueueInputBuffer)
+#define MLMediaCodecQueueInputBuffer ::MLSDK_API::MLMediaCodecQueueInputBufferShim
+CREATE_FUNCTION_SHIM(ml_mediacodec, MLResult, MLMediaCodecQueueSecureInputBuffer)
+#define MLMediaCodecQueueSecureInputBuffer ::MLSDK_API::MLMediaCodecQueueSecureInputBufferShim
+CREATE_FUNCTION_SHIM(ml_mediacodec, MLResult, MLMediaCodecDequeueInputBuffer)
+#define MLMediaCodecDequeueInputBuffer ::MLSDK_API::MLMediaCodecDequeueInputBufferShim
+CREATE_FUNCTION_SHIM(ml_mediacodec, MLResult, MLMediaCodecDequeueOutputBuffer)
+#define MLMediaCodecDequeueOutputBuffer ::MLSDK_API::MLMediaCodecDequeueOutputBufferShim
+CREATE_FUNCTION_SHIM(ml_mediacodec, MLResult, MLMediaCodecGetInputFormat)
+#define MLMediaCodecGetInputFormat ::MLSDK_API::MLMediaCodecGetInputFormatShim
+CREATE_FUNCTION_SHIM(ml_mediacodec, MLResult, MLMediaCodecGetOutputFormat)
+#define MLMediaCodecGetOutputFormat ::MLSDK_API::MLMediaCodecGetOutputFormatShim
+CREATE_FUNCTION_SHIM(ml_mediacodec, MLResult, MLMediaCodecReleaseOutputBuffer)
+#define MLMediaCodecReleaseOutputBuffer ::MLSDK_API::MLMediaCodecReleaseOutputBufferShim
+CREATE_FUNCTION_SHIM(ml_mediacodec, MLResult, MLMediaCodecReleaseOutputBufferAtTime)
+#define MLMediaCodecReleaseOutputBufferAtTime ::MLSDK_API::MLMediaCodecReleaseOutputBufferAtTimeShim
+CREATE_FUNCTION_SHIM(ml_mediacodec, MLResult, MLMediaCodecAcquireNextAvailableFrame)
+#define MLMediaCodecAcquireNextAvailableFrame ::MLSDK_API::MLMediaCodecAcquireNextAvailableFrameShim
+CREATE_FUNCTION_SHIM(ml_mediacodec, MLResult, MLMediaCodecReleaseFrame)
+#define MLMediaCodecReleaseFrame ::MLSDK_API::MLMediaCodecReleaseFrameShim
+CREATE_FUNCTION_SHIM(ml_mediacodec, MLResult, MLMediaCodecGetFrameTransformationMatrix)
+#define MLMediaCodecGetFrameTransformationMatrix ::MLSDK_API::MLMediaCodecGetFrameTransformationMatrixShim
+CREATE_FUNCTION_SHIM(ml_mediacodec, MLResult, MLMediaCodecGetFrameTimestamp)
+#define MLMediaCodecGetFrameTimestamp ::MLSDK_API::MLMediaCodecGetFrameTimestampShim
+CREATE_FUNCTION_SHIM(ml_mediacodec, MLResult, MLMediaCodecGetFrameQueueBufferTimestamp)
+#define MLMediaCodecGetFrameQueueBufferTimestamp ::MLSDK_API::MLMediaCodecGetFrameQueueBufferTimestampShim
+CREATE_FUNCTION_SHIM(ml_mediacodec, MLResult, MLMediaCodecGetFrameNumber)
+#define MLMediaCodecGetFrameNumber ::MLSDK_API::MLMediaCodecGetFrameNumberShim
+
+}
+
+#endif // !defined(WITH_MLSDK) || WITH_MLSDK
diff --git a/Engine/Plugins/Lumin/MLSDK/Source/MLSDK/Public/Lumin/CAPIShims/LuminAPIMediaCodecList.h b/Engine/Plugins/Lumin/MLSDK/Source/MLSDK/Public/Lumin/CAPIShims/LuminAPIMediaCodecList.h
new file mode 100644
--- /dev/null
+++ b/Engine/Plugins/Lumin/MLSDK/Source/MLSDK/Public/Lumin/CAPIShims/LuminAPIMediaCodecList.h
@@ -0,0 +1,42 @@
+// Copyright 1998-2019 Epic Games, Inc. All Rights Reserved.
+
+#pragma once
+
+#if !defined(WITH_MLSDK) || WITH_MLSDK
+
+#include "Lumin/CAPIShims/LuminAPI.h"
+
+LUMIN_THIRD_PARTY_INCLUDES_START
+#include <ml_media_common.h>
+#include <ml_media_codeclist.h>
+LUMIN_THIRD_PARTY_INCLUDES_END
+
+namespace MLSDK_API
+{
+
+CREATE_FUNCTION_SHIM(ml_mediacodeclist, MLResult, MLMediaCodecListCountCodecs)
+#define MLMediaCodecListCountCodecs ::MLSDK_API::MLMediaCodecListCountCodecsShim
+CREATE_FUNCTION_SHIM(ml_mediacodeclist, MLResult, MLMediaCodecListIsSoftwareCodec)
+#define MLMediaCodecListIsSoftwareCodec ::MLSDK_API::MLMediaCodecListIsSoftwareCodecShim
+CREATE_FUNCTION_SHIM(ml_mediacodeclist, MLResult, MLMediaCodecListGetMatchingCodecs)
+#define MLMediaCodecListGetMatchingCodecs ::MLSDK_API::MLMediaCodecListGetMatchingCodecsShim
+CREATE_FUNCTION_SHIM(ml_mediacodeclist, MLResult, MLMediaCodecListGetCodecByType)
+#define MLMediaCodecListGetCodecByType ::MLSDK_API::MLMediaCodecListGetCodecByTypeShim
+CREATE_FUNCTION_SHIM(ml_mediacodeclist, MLResult, MLMediaCodecListGetCodecByName)
+#define MLMediaCodecListGetCodecByName ::MLSDK_API::MLMediaCodecListGetCodecByNameShim
+CREATE_FUNCTION_SHIM(ml_mediacodeclist, MLResult, MLMediaCodecListIsEncoder)
+#define MLMediaCodecListIsEncoder ::MLSDK_API::MLMediaCodecListIsEncoderShim
+CREATE_FUNCTION_SHIM(ml_mediacodeclist, MLResult, MLMediaCodecListGetCodecName)
+#define MLMediaCodecListGetCodecName ::MLSDK_API::MLMediaCodecListGetCodecNameShim
+CREATE_FUNCTION_SHIM(ml_mediacodeclist, MLResult, MLMediaCodecListGetSupportedMimes)
+#define MLMediaCodecListGetSupportedMimes ::MLSDK_API::MLMediaCodecListGetSupportedMimesShim
+CREATE_FUNCTION_SHIM(ml_mediacodeclist, MLResult, MLMediaCodecListGetCapabilityFlag)
+#define MLMediaCodecListGetCapabilityFlag ::MLSDK_API::MLMediaCodecListGetCapabilityFlagShim
+CREATE_FUNCTION_SHIM(ml_mediacodeclist, MLResult, MLMediaCodecListGetSecureCodecName)
+#define MLMediaCodecListGetSecureCodecName ::MLSDK_API::MLMediaCodecListGetSecureCodecNameShim
+CREATE_FUNCTION_SHIM(ml_mediacodeclist, MLResult, MLMediaCodecListQueryResultsRelease)
+#define MLMediaCodecListQueryResultsRelease ::MLSDK_API::MLMediaCodecListQueryResultsReleaseShim
+
+}
+
+#endif // !defined(WITH_MLSDK) || WITH_MLSDK
diff --git a/Engine/Plugins/Lumin/MLSDK/Source/MLSDK/Public/Lumin/CAPIShims/LuminAPIMediaCrypto.h b/Engine/Plugins/Lumin/MLSDK/Source/MLSDK/Public/Lumin/CAPIShims/LuminAPIMediaCrypto.h
new file mode 100644
--- /dev/null
+++ b/Engine/Plugins/Lumin/MLSDK/Source/MLSDK/Public/Lumin/CAPIShims/LuminAPIMediaCrypto.h
@@ -0,0 +1,30 @@
+// Copyright 1998-2019 Epic Games, Inc. All Rights Reserved.
+
+#pragma once
+
+#if !defined(WITH_MLSDK) || WITH_MLSDK
+
+#include "Lumin/CAPIShims/LuminAPI.h"
+
+LUMIN_THIRD_PARTY_INCLUDES_START
+#include <ml_media_common.h>
+#include <ml_media_crypto.h>
+LUMIN_THIRD_PARTY_INCLUDES_END
+
+namespace MLSDK_API
+{
+
+CREATE_FUNCTION_SHIM(ml_mediacrypto, MLResult, MLMediaCryptoCreate)
+#define MLMediaCryptoCreate ::MLSDK_API::MLMediaCryptoCreateShim
+CREATE_FUNCTION_SHIM(ml_mediacrypto, MLResult, MLMediaCryptoRelease)
+#define MLMediaCryptoRelease ::MLSDK_API::MLMediaCryptoReleaseShim
+CREATE_FUNCTION_SHIM(ml_mediacrypto, MLResult, MLMediaCryptoIsCryptoSchemeSupported)
+#define MLMediaCryptoIsCryptoSchemeSupported ::MLSDK_API::MLMediaCryptoIsCryptoSchemeSupportedShim
+CREATE_FUNCTION_SHIM(ml_mediacrypto, MLResult, MLMediaCryptoRequiresSecureDecoderComponent)
+#define MLMediaCryptoRequiresSecureDecoderComponent ::MLSDK_API::MLMediaCryptoRequiresSecureDecoderComponentShim
+CREATE_FUNCTION_SHIM(ml_mediacrypto, MLResult, MLMediaCryptoSetMediaDRMSession)
+#define MLMediaCryptoSetMediaDRMSession ::MLSDK_API::MLMediaCryptoSetMediaDRMSessionShim
+
+}
+
+#endif // !defined(WITH_MLSDK) || WITH_MLSDK
diff --git a/Engine/Plugins/Lumin/MLSDK/Source/MLSDK/Public/Lumin/CAPIShims/LuminAPIMediaDRM.h b/Engine/Plugins/Lumin/MLSDK/Source/MLSDK/Public/Lumin/CAPIShims/LuminAPIMediaDRM.h
new file mode 100644
--- /dev/null
+++ b/Engine/Plugins/Lumin/MLSDK/Source/MLSDK/Public/Lumin/CAPIShims/LuminAPIMediaDRM.h
@@ -0,0 +1,88 @@
+// Copyright 1998-2019 Epic Games, Inc. All Rights Reserved.
+
+#pragma once
+
+#if !defined(WITH_MLSDK) || WITH_MLSDK
+
+#include "Lumin/CAPIShims/LuminAPI.h"
+
+LUMIN_THIRD_PARTY_INCLUDES_START
+#include <ml_media_common.h>
+#include <ml_media_drm.h>
+LUMIN_THIRD_PARTY_INCLUDES_END
+
+namespace MLSDK_API
+{
+
+CREATE_FUNCTION_SHIM(ml_mediadrm, MLResult, MLMediaDRMByteArrayAllocate)
+#define MLMediaDRMByteArrayAllocate ::MLSDK_API::MLMediaDRMByteArrayAllocateShim
+CREATE_FUNCTION_SHIM(ml_mediadrm, MLResult, MLMediaDRMByteArrayAllocAndCopy)
+#define MLMediaDRMByteArrayAllocAndCopy ::MLSDK_API::MLMediaDRMByteArrayAllocAndCopyShim
+CREATE_FUNCTION_SHIM(ml_mediadrm, MLResult, MLMediaDRMKeyValueArrayAllocate)
+#define MLMediaDRMKeyValueArrayAllocate ::MLSDK_API::MLMediaDRMKeyValueArrayAllocateShim
+CREATE_FUNCTION_SHIM(ml_mediadrm, MLResult, MLMediaDRMCreate)
+#define MLMediaDRMCreate ::MLSDK_API::MLMediaDRMCreateShim
+CREATE_FUNCTION_SHIM(ml_mediadrm, MLResult, MLMediaDRMRelease)
+#define MLMediaDRMRelease ::MLSDK_API::MLMediaDRMReleaseShim
+CREATE_FUNCTION_SHIM(ml_mediadrm, MLResult, MLMediaDRMByteArrayRelease)
+#define MLMediaDRMByteArrayRelease ::MLSDK_API::MLMediaDRMByteArrayReleaseShim
+CREATE_FUNCTION_SHIM(ml_mediadrm, MLResult, MLMediaDRMByteArrayListRelease)
+#define MLMediaDRMByteArrayListRelease ::MLSDK_API::MLMediaDRMByteArrayListReleaseShim
+CREATE_FUNCTION_SHIM(ml_mediadrm, MLResult, MLMediaDRMKeyValueArrayRelease)
+#define MLMediaDRMKeyValueArrayRelease ::MLSDK_API::MLMediaDRMKeyValueArrayReleaseShim
+CREATE_FUNCTION_SHIM(ml_mediadrm, MLResult, MLMediaDRMRequestMessageRelease)
+#define MLMediaDRMRequestMessageRelease ::MLSDK_API::MLMediaDRMRequestMessageReleaseShim
+CREATE_FUNCTION_SHIM(ml_mediadrm, MLResult, MLMediaDRMKeyValueArrayAdd)
+#define MLMediaDRMKeyValueArrayAdd ::MLSDK_API::MLMediaDRMKeyValueArrayAddShim
+CREATE_FUNCTION_SHIM(ml_mediadrm, MLResult, MLMediaDRMIsCryptoSchemeSupported)
+#define MLMediaDRMIsCryptoSchemeSupported ::MLSDK_API::MLMediaDRMIsCryptoSchemeSupportedShim
+CREATE_FUNCTION_SHIM(ml_mediadrm, MLResult, MLMediaDRMSetOnEventListener)
+#define MLMediaDRMSetOnEventListener ::MLSDK_API::MLMediaDRMSetOnEventListenerShim
+CREATE_FUNCTION_SHIM(ml_mediadrm, MLResult, MLMediaDRMOpenSession)
+#define MLMediaDRMOpenSession ::MLSDK_API::MLMediaDRMOpenSessionShim
+CREATE_FUNCTION_SHIM(ml_mediadrm, MLResult, MLMediaDRMCloseSession)
+#define MLMediaDRMCloseSession ::MLSDK_API::MLMediaDRMCloseSessionShim
+CREATE_FUNCTION_SHIM(ml_mediadrm, MLResult, MLMediaDRMGetKeyRequest)
+#define MLMediaDRMGetKeyRequest ::MLSDK_API::MLMediaDRMGetKeyRequestShim
+CREATE_FUNCTION_SHIM(ml_mediadrm, MLResult, MLMediaDRMProvideKeyResponse)
+#define MLMediaDRMProvideKeyResponse ::MLSDK_API::MLMediaDRMProvideKeyResponseShim
+CREATE_FUNCTION_SHIM(ml_mediadrm, MLResult, MLMediaDRMRestoreKeys)
+#define MLMediaDRMRestoreKeys ::MLSDK_API::MLMediaDRMRestoreKeysShim
+CREATE_FUNCTION_SHIM(ml_mediadrm, MLResult, MLMediaDRMRemoveKeys)
+#define MLMediaDRMRemoveKeys ::MLSDK_API::MLMediaDRMRemoveKeysShim
+CREATE_FUNCTION_SHIM(ml_mediadrm, MLResult, MLMediaDRMQueryKeyStatus)
+#define MLMediaDRMQueryKeyStatus ::MLSDK_API::MLMediaDRMQueryKeyStatusShim
+CREATE_FUNCTION_SHIM(ml_mediadrm, MLResult, MLMediaDRMGetProvisionRequest)
+#define MLMediaDRMGetProvisionRequest ::MLSDK_API::MLMediaDRMGetProvisionRequestShim
+CREATE_FUNCTION_SHIM(ml_mediadrm, MLResult, MLMediaDRMProvideProvisionResponse)
+#define MLMediaDRMProvideProvisionResponse ::MLSDK_API::MLMediaDRMProvideProvisionResponseShim
+CREATE_FUNCTION_SHIM(ml_mediadrm, MLResult, MLMediaDRMGetSecureStops)
+#define MLMediaDRMGetSecureStops ::MLSDK_API::MLMediaDRMGetSecureStopsShim
+CREATE_FUNCTION_SHIM(ml_mediadrm, MLResult, MLMediaDRMGetSecureStop)
+#define MLMediaDRMGetSecureStop ::MLSDK_API::MLMediaDRMGetSecureStopShim
+CREATE_FUNCTION_SHIM(ml_mediadrm, MLResult, MLMediaDRMReleaseSecureStops)
+#define MLMediaDRMReleaseSecureStops ::MLSDK_API::MLMediaDRMReleaseSecureStopsShim
+CREATE_FUNCTION_SHIM(ml_mediadrm, MLResult, MLMediaDRMReleaseAllSecureStops)
+#define MLMediaDRMReleaseAllSecureStops ::MLSDK_API::MLMediaDRMReleaseAllSecureStopsShim
+CREATE_FUNCTION_SHIM(ml_mediadrm, MLResult, MLMediaDRMGetPropertyString)
+#define MLMediaDRMGetPropertyString ::MLSDK_API::MLMediaDRMGetPropertyStringShim
+CREATE_FUNCTION_SHIM(ml_mediadrm, MLResult, MLMediaDRMGetPropertyByteArray)
+#define MLMediaDRMGetPropertyByteArray ::MLSDK_API::MLMediaDRMGetPropertyByteArrayShim
+CREATE_FUNCTION_SHIM(ml_mediadrm, MLResult, MLMediaDRMSetPropertyString)
+#define MLMediaDRMSetPropertyString ::MLSDK_API::MLMediaDRMSetPropertyStringShim
+CREATE_FUNCTION_SHIM(ml_mediadrm, MLResult, MLMediaDRMSetPropertyByteArray)
+#define MLMediaDRMSetPropertyByteArray ::MLSDK_API::MLMediaDRMSetPropertyByteArrayShim
+CREATE_FUNCTION_SHIM(ml_mediadrm, MLResult, MLMediaDRMEncrypt)
+#define MLMediaDRMEncrypt ::MLSDK_API::MLMediaDRMEncryptShim
+CREATE_FUNCTION_SHIM(ml_mediadrm, MLResult, MLMediaDRMDecrypt)
+#define MLMediaDRMDecrypt ::MLSDK_API::MLMediaDRMDecryptShim
+CREATE_FUNCTION_SHIM(ml_mediadrm, MLResult, MLMediaDRMSign)
+#define MLMediaDRMSign ::MLSDK_API::MLMediaDRMSignShim
+CREATE_FUNCTION_SHIM(ml_mediadrm, MLResult, MLMediaDRMVerify)
+#define MLMediaDRMVerify ::MLSDK_API::MLMediaDRMVerifyShim
+CREATE_FUNCTION_SHIM(ml_mediadrm, MLResult, MLMediaDRMSignRSA)
+#define MLMediaDRMSignRSA ::MLSDK_API::MLMediaDRMSignRSAShim
+
+}
+
+#endif // !defined(WITH_MLSDK) || WITH_MLSDK
diff --git a/Engine/Plugins/Lumin/MLSDK/Source/MLSDK/Public/Lumin/CAPIShims/LuminAPIMediaDataSource.h b/Engine/Plugins/Lumin/MLSDK/Source/MLSDK/Public/Lumin/CAPIShims/LuminAPIMediaDataSource.h
new file mode 100644
--- /dev/null
+++ b/Engine/Plugins/Lumin/MLSDK/Source/MLSDK/Public/Lumin/CAPIShims/LuminAPIMediaDataSource.h
@@ -0,0 +1,24 @@
+// Copyright 1998-2019 Epic Games, Inc. All Rights Reserved.
+
+#pragma once
+
+#if !defined(WITH_MLSDK) || WITH_MLSDK
+
+#include "Lumin/CAPIShims/LuminAPI.h"
+
+LUMIN_THIRD_PARTY_INCLUDES_START
+#include <ml_media_common.h>
+#include <ml_media_data_source.h>
+LUMIN_THIRD_PARTY_INCLUDES_END
+
+namespace MLSDK_API
+{
+
+CREATE_FUNCTION_SHIM(ml_mediaplayer, MLResult, MLMediaDataSourceCreate)
+#define MLMediaDataSourceCreate ::MLSDK_API::MLMediaDataSourceCreateShim
+CREATE_FUNCTION_SHIM(ml_mediaplayer, MLResult, MLMediaDataSourceDestroy)
+#define MLMediaDataSourceDestroy ::MLSDK_API::MLMediaDataSourceDestroyShim
+
+}
+
+#endif // !defined(WITH_MLSDK) || WITH_MLSDK
diff --git a/Engine/Plugins/Lumin/MLSDK/Source/MLSDK/Public/Lumin/CAPIShims/LuminAPIMediaError.h b/Engine/Plugins/Lumin/MLSDK/Source/MLSDK/Public/Lumin/CAPIShims/LuminAPIMediaError.h
new file mode 100644
--- /dev/null
+++ b/Engine/Plugins/Lumin/MLSDK/Source/MLSDK/Public/Lumin/CAPIShims/LuminAPIMediaError.h
@@ -0,0 +1,22 @@
+// Copyright 1998-2019 Epic Games, Inc. All Rights Reserved.
+
+#pragma once
+
+#if !defined(WITH_MLSDK) || WITH_MLSDK
+
+#include "Lumin/CAPIShims/LuminAPI.h"
+
+LUMIN_THIRD_PARTY_INCLUDES_START
+#include <ml_media_common.h>
+#include <ml_media_error.h>
+LUMIN_THIRD_PARTY_INCLUDES_END
+
+namespace MLSDK_API
+{
+
+CREATE_FUNCTION_SHIM(ml_mediaerror, const char *, MLMediaResultGetString)
+#define MLMediaResultGetString ::MLSDK_API::MLMediaResultGetStringShim
+
+}
+
+#endif // !defined(WITH_MLSDK) || WITH_MLSDK
diff --git a/Engine/Plugins/Lumin/MLSDK/Source/MLSDK/Public/Lumin/CAPIShims/LuminAPIMediaExtractor.h b/Engine/Plugins/Lumin/MLSDK/Source/MLSDK/Public/Lumin/CAPIShims/LuminAPIMediaExtractor.h
new file mode 100644
--- /dev/null
+++ b/Engine/Plugins/Lumin/MLSDK/Source/MLSDK/Public/Lumin/CAPIShims/LuminAPIMediaExtractor.h
@@ -0,0 +1,72 @@
+// Copyright 1998-2019 Epic Games, Inc. All Rights Reserved.
+
+#pragma once
+
+#if !defined(WITH_MLSDK) || WITH_MLSDK
+
+#include "Lumin/CAPIShims/LuminAPI.h"
+
+LUMIN_THIRD_PARTY_INCLUDES_START
+#include <ml_media_common.h>
+#include <ml_media_extractor.h>
+LUMIN_THIRD_PARTY_INCLUDES_END
+
+namespace MLSDK_API
+{
+
+CREATE_FUNCTION_SHIM(ml_mediaextractor, MLResult, MLMediaExtractorCreate)
+#define MLMediaExtractorCreate ::MLSDK_API::MLMediaExtractorCreateShim
+CREATE_FUNCTION_SHIM(ml_mediaextractor, MLResult, MLMediaExtractorDestroy)
+#define MLMediaExtractorDestroy ::MLSDK_API::MLMediaExtractorDestroyShim
+CREATE_FUNCTION_SHIM(ml_mediaextractor, MLResult, MLMediaExtractorSetDataSourceForPath)
+#define MLMediaExtractorSetDataSourceForPath ::MLSDK_API::MLMediaExtractorSetDataSourceForPathShim
+CREATE_FUNCTION_SHIM(ml_mediaextractor, MLResult, MLMediaExtractorSetDataSourceForURI)
+#define MLMediaExtractorSetDataSourceForURI ::MLSDK_API::MLMediaExtractorSetDataSourceForURIShim
+CREATE_FUNCTION_SHIM(ml_mediaextractor, MLResult, MLMediaExtractorSetDataSourceForURIWithHeaders)
+#define MLMediaExtractorSetDataSourceForURIWithHeaders ::MLSDK_API::MLMediaExtractorSetDataSourceForURIWithHeadersShim
+CREATE_FUNCTION_SHIM(ml_mediaextractor, MLResult, MLMediaExtractorSetDataSourceForFD)
+#define MLMediaExtractorSetDataSourceForFD ::MLSDK_API::MLMediaExtractorSetDataSourceForFDShim
+CREATE_FUNCTION_SHIM(ml_mediaextractor, MLResult, MLMediaExtractorSetMediaDataSource)
+#define MLMediaExtractorSetMediaDataSource ::MLSDK_API::MLMediaExtractorSetMediaDataSourceShim
+CREATE_FUNCTION_SHIM(ml_mediaextractor, MLResult, MLMediaExtractorGetDRMInitData)
+#define MLMediaExtractorGetDRMInitData ::MLSDK_API::MLMediaExtractorGetDRMInitDataShim
+CREATE_FUNCTION_SHIM(ml_mediaextractor, MLResult, MLMediaExtractorReleaseDRMInitData)
+#define MLMediaExtractorReleaseDRMInitData ::MLSDK_API::MLMediaExtractorReleaseDRMInitDataShim
+CREATE_FUNCTION_SHIM(ml_mediaextractor, MLResult, MLMediaExtractorGetTrackCount)
+#define MLMediaExtractorGetTrackCount ::MLSDK_API::MLMediaExtractorGetTrackCountShim
+CREATE_FUNCTION_SHIM(ml_mediaextractor, MLResult, MLMediaExtractorGetTrackFormat)
+#define MLMediaExtractorGetTrackFormat ::MLSDK_API::MLMediaExtractorGetTrackFormatShim
+CREATE_FUNCTION_SHIM(ml_mediaextractor, MLResult, MLMediaExtractorSelectTrack)
+#define MLMediaExtractorSelectTrack ::MLSDK_API::MLMediaExtractorSelectTrackShim
+CREATE_FUNCTION_SHIM(ml_mediaextractor, MLResult, MLMediaExtractorUnselectTrack)
+#define MLMediaExtractorUnselectTrack ::MLSDK_API::MLMediaExtractorUnselectTrackShim
+CREATE_FUNCTION_SHIM(ml_mediaextractor, MLResult, MLMediaExtractorSeekTo)
+#define MLMediaExtractorSeekTo ::MLSDK_API::MLMediaExtractorSeekToShim
+CREATE_FUNCTION_SHIM(ml_mediaextractor, MLResult, MLMediaExtractorAdvance)
+#define MLMediaExtractorAdvance ::MLSDK_API::MLMediaExtractorAdvanceShim
+CREATE_FUNCTION_SHIM(ml_mediaextractor, MLResult, MLMediaExtractorReadSampleData)
+#define MLMediaExtractorReadSampleData ::MLSDK_API::MLMediaExtractorReadSampleDataShim
+CREATE_FUNCTION_SHIM(ml_mediaextractor, MLResult, MLMediaExtractorGetSampleTrackIndex)
+#define MLMediaExtractorGetSampleTrackIndex ::MLSDK_API::MLMediaExtractorGetSampleTrackIndexShim
+CREATE_FUNCTION_SHIM(ml_mediaextractor, MLResult, MLMediaExtractorGetSampleTime)
+#define MLMediaExtractorGetSampleTime ::MLSDK_API::MLMediaExtractorGetSampleTimeShim
+CREATE_FUNCTION_SHIM(ml_mediaextractor, MLResult, MLMediaExtractorGetSampleFlags)
+#define MLMediaExtractorGetSampleFlags ::MLSDK_API::MLMediaExtractorGetSampleFlagsShim
+CREATE_FUNCTION_SHIM(ml_mediaextractor, MLResult, MLMediaExtractorGetSampleCryptoInfo)
+#define MLMediaExtractorGetSampleCryptoInfo ::MLSDK_API::MLMediaExtractorGetSampleCryptoInfoShim
+CREATE_FUNCTION_SHIM(ml_mediaextractor, MLResult, MLMediaExtractorReleaseCryptoInfo)
+#define MLMediaExtractorReleaseCryptoInfo ::MLSDK_API::MLMediaExtractorReleaseCryptoInfoShim
+CREATE_FUNCTION_SHIM(ml_mediaextractor, MLResult, MLMediaExtractorGetCachedDuration)
+#define MLMediaExtractorGetCachedDuration ::MLSDK_API::MLMediaExtractorGetCachedDurationShim
+CREATE_FUNCTION_SHIM(ml_mediaextractor, MLResult, MLMediaExtractorHasCacheReachedEOS)
+#define MLMediaExtractorHasCacheReachedEOS ::MLSDK_API::MLMediaExtractorHasCacheReachedEOSShim
+CREATE_FUNCTION_SHIM(ml_mediaextractor, MLResult, MLMediaExtractorGetPSSHInfoCount)
+#define MLMediaExtractorGetPSSHInfoCount ::MLSDK_API::MLMediaExtractorGetPSSHInfoCountShim
+CREATE_FUNCTION_SHIM(ml_mediaextractor, MLResult, MLMediaExtractorGetPSSHInfo)
+#define MLMediaExtractorGetPSSHInfo ::MLSDK_API::MLMediaExtractorGetPSSHInfoShim
+CREATE_FUNCTION_SHIM(ml_mediaextractor, MLResult, MLMediaExtractorReleasePSSHInfo)
+#define MLMediaExtractorReleasePSSHInfo ::MLSDK_API::MLMediaExtractorReleasePSSHInfoShim
+
+}
+
+#endif // !defined(WITH_MLSDK) || WITH_MLSDK
diff --git a/Engine/Plugins/Lumin/MLSDK/Source/MLSDK/Public/Lumin/CAPIShims/LuminAPIMediaFormat.h b/Engine/Plugins/Lumin/MLSDK/Source/MLSDK/Public/Lumin/CAPIShims/LuminAPIMediaFormat.h
new file mode 100644
--- /dev/null
+++ b/Engine/Plugins/Lumin/MLSDK/Source/MLSDK/Public/Lumin/CAPIShims/LuminAPIMediaFormat.h
@@ -0,0 +1,135 @@
+// Copyright 1998-2019 Epic Games, Inc. All Rights Reserved.
+
+#pragma once
+
+#if !defined(WITH_MLSDK) || WITH_MLSDK
+
+#include "Lumin/CAPIShims/LuminAPI.h"
+
+LUMIN_THIRD_PARTY_INCLUDES_START
+#include <ml_media_common.h>
+#include <ml_media_format.h>
+LUMIN_THIRD_PARTY_INCLUDES_END
+
+namespace MLSDK_API
+{
+
+CREATE_GLOBAL_SHIM(ml_mediaformat, MLMediaFormatKey, MLMediaFormat_Key_AAC_Encoded_Target_Level)
+#define MLMediaFormat_Key_AAC_Encoded_Target_Level ::MLSDK_API::MLMediaFormat_Key_AAC_Encoded_Target_LevelShim()
+CREATE_GLOBAL_SHIM(ml_mediaformat, MLMediaFormatKey, MLMediaFormat_Key_AAC_DRC_Boost_Factor)
+#define MLMediaFormat_Key_AAC_DRC_Boost_Factor ::MLSDK_API::MLMediaFormat_Key_AAC_DRC_Boost_FactorShim()
+CREATE_GLOBAL_SHIM(ml_mediaformat, MLMediaFormatKey, MLMediaFormat_Key_AAC_DRC_Attenuation_Factor)
+#define MLMediaFormat_Key_AAC_DRC_Attenuation_Factor ::MLSDK_API::MLMediaFormat_Key_AAC_DRC_Attenuation_FactorShim()
+CREATE_GLOBAL_SHIM(ml_mediaformat, MLMediaFormatKey, MLMediaFormat_Key_AAC_DRC_Heavy_Compression)
+#define MLMediaFormat_Key_AAC_DRC_Heavy_Compression ::MLSDK_API::MLMediaFormat_Key_AAC_DRC_Heavy_CompressionShim()
+CREATE_GLOBAL_SHIM(ml_mediaformat, MLMediaFormatKey, MLMediaFormat_Key_AAC_DRC_Target_Reference_Level)
+#define MLMediaFormat_Key_AAC_DRC_Target_Reference_Level ::MLSDK_API::MLMediaFormat_Key_AAC_DRC_Target_Reference_LevelShim()
+CREATE_GLOBAL_SHIM(ml_mediaformat, MLMediaFormatKey, MLMediaFormat_Key_AAC_Max_Output_Channel_Count)
+#define MLMediaFormat_Key_AAC_Max_Output_Channel_Count ::MLSDK_API::MLMediaFormat_Key_AAC_Max_Output_Channel_CountShim()
+CREATE_GLOBAL_SHIM(ml_mediaformat, MLMediaFormatKey, MLMediaFormat_Key_AAC_Profile)
+#define MLMediaFormat_Key_AAC_Profile ::MLSDK_API::MLMediaFormat_Key_AAC_ProfileShim()
+CREATE_GLOBAL_SHIM(ml_mediaformat, MLMediaFormatKey, MLMediaFormat_Key_AAC_SBR_Mode)
+#define MLMediaFormat_Key_AAC_SBR_Mode ::MLSDK_API::MLMediaFormat_Key_AAC_SBR_ModeShim()
+CREATE_GLOBAL_SHIM(ml_mediaformat, MLMediaFormatKey, MLMediaFormat_Key_Bit_Rate)
+#define MLMediaFormat_Key_Bit_Rate ::MLSDK_API::MLMediaFormat_Key_Bit_RateShim()
+CREATE_GLOBAL_SHIM(ml_mediaformat, MLMediaFormatKey, MLMediaFormat_Key_Capture_Rate)
+#define MLMediaFormat_Key_Capture_Rate ::MLSDK_API::MLMediaFormat_Key_Capture_RateShim()
+CREATE_GLOBAL_SHIM(ml_mediaformat, MLMediaFormatKey, MLMediaFormat_Key_Channel_Count)
+#define MLMediaFormat_Key_Channel_Count ::MLSDK_API::MLMediaFormat_Key_Channel_CountShim()
+CREATE_GLOBAL_SHIM(ml_mediaformat, MLMediaFormatKey, MLMediaFormat_Key_Channel_Mask)
+#define MLMediaFormat_Key_Channel_Mask ::MLSDK_API::MLMediaFormat_Key_Channel_MaskShim()
+CREATE_GLOBAL_SHIM(ml_mediaformat, MLMediaFormatKey, MLMediaFormat_Key_Color_Format)
+#define MLMediaFormat_Key_Color_Format ::MLSDK_API::MLMediaFormat_Key_Color_FormatShim()
+CREATE_GLOBAL_SHIM(ml_mediaformat, MLMediaFormatKey, MLMediaFormat_Key_Duration)
+#define MLMediaFormat_Key_Duration ::MLSDK_API::MLMediaFormat_Key_DurationShim()
+CREATE_GLOBAL_SHIM(ml_mediaformat, MLMediaFormatKey, MLMediaFormat_Key_FLAC_Compression_Level)
+#define MLMediaFormat_Key_FLAC_Compression_Level ::MLSDK_API::MLMediaFormat_Key_FLAC_Compression_LevelShim()
+CREATE_GLOBAL_SHIM(ml_mediaformat, MLMediaFormatKey, MLMediaFormat_Key_Frame_Rate)
+#define MLMediaFormat_Key_Frame_Rate ::MLSDK_API::MLMediaFormat_Key_Frame_RateShim()
+CREATE_GLOBAL_SHIM(ml_mediaformat, MLMediaFormatKey, MLMediaFormat_Key_Height)
+#define MLMediaFormat_Key_Height ::MLSDK_API::MLMediaFormat_Key_HeightShim()
+CREATE_GLOBAL_SHIM(ml_mediaformat, MLMediaFormatKey, MLMediaFormat_Key_I_Frame_Interval)
+#define MLMediaFormat_Key_I_Frame_Interval ::MLSDK_API::MLMediaFormat_Key_I_Frame_IntervalShim()
+CREATE_GLOBAL_SHIM(ml_mediaformat, MLMediaFormatKey, MLMediaFormat_Key_Intra_Refresh_Period)
+#define MLMediaFormat_Key_Intra_Refresh_Period ::MLSDK_API::MLMediaFormat_Key_Intra_Refresh_PeriodShim()
+CREATE_GLOBAL_SHIM(ml_mediaformat, MLMediaFormatKey, MLMediaFormat_Key_Is_ADTS)
+#define MLMediaFormat_Key_Is_ADTS ::MLSDK_API::MLMediaFormat_Key_Is_ADTSShim()
+CREATE_GLOBAL_SHIM(ml_mediaformat, MLMediaFormatKey, MLMediaFormat_Key_Is_Autoselect)
+#define MLMediaFormat_Key_Is_Autoselect ::MLSDK_API::MLMediaFormat_Key_Is_AutoselectShim()
+CREATE_GLOBAL_SHIM(ml_mediaformat, MLMediaFormatKey, MLMediaFormat_Key_Is_Default)
+#define MLMediaFormat_Key_Is_Default ::MLSDK_API::MLMediaFormat_Key_Is_DefaultShim()
+CREATE_GLOBAL_SHIM(ml_mediaformat, MLMediaFormatKey, MLMediaFormat_Key_Is_Forced_Subtitle)
+#define MLMediaFormat_Key_Is_Forced_Subtitle ::MLSDK_API::MLMediaFormat_Key_Is_Forced_SubtitleShim()
+CREATE_GLOBAL_SHIM(ml_mediaformat, MLMediaFormatKey, MLMediaFormat_Key_Language)
+#define MLMediaFormat_Key_Language ::MLSDK_API::MLMediaFormat_Key_LanguageShim()
+CREATE_GLOBAL_SHIM(ml_mediaformat, MLMediaFormatKey, MLMediaFormat_Key_Max_Height)
+#define MLMediaFormat_Key_Max_Height ::MLSDK_API::MLMediaFormat_Key_Max_HeightShim()
+CREATE_GLOBAL_SHIM(ml_mediaformat, MLMediaFormatKey, MLMediaFormat_Key_Max_Input_Size)
+#define MLMediaFormat_Key_Max_Input_Size ::MLSDK_API::MLMediaFormat_Key_Max_Input_SizeShim()
+CREATE_GLOBAL_SHIM(ml_mediaformat, MLMediaFormatKey, MLMediaFormat_Key_Max_Width)
+#define MLMediaFormat_Key_Max_Width ::MLSDK_API::MLMediaFormat_Key_Max_WidthShim()
+CREATE_GLOBAL_SHIM(ml_mediaformat, MLMediaFormatKey, MLMediaFormat_Key_Mime)
+#define MLMediaFormat_Key_Mime ::MLSDK_API::MLMediaFormat_Key_MimeShim()
+CREATE_GLOBAL_SHIM(ml_mediaformat, MLMediaFormatKey, MLMediaFormat_Key_PCM_Encoding)
+#define MLMediaFormat_Key_PCM_Encoding ::MLSDK_API::MLMediaFormat_Key_PCM_EncodingShim()
+CREATE_GLOBAL_SHIM(ml_mediaformat, MLMediaFormatKey, MLMediaFormat_Key_Push_Blank_Buffers_On_Stop)
+#define MLMediaFormat_Key_Push_Blank_Buffers_On_Stop ::MLSDK_API::MLMediaFormat_Key_Push_Blank_Buffers_On_StopShim()
+CREATE_GLOBAL_SHIM(ml_mediaformat, MLMediaFormatKey, MLMediaFormat_Key_Repeat_Previous_Frame_After)
+#define MLMediaFormat_Key_Repeat_Previous_Frame_After ::MLSDK_API::MLMediaFormat_Key_Repeat_Previous_Frame_AfterShim()
+CREATE_GLOBAL_SHIM(ml_mediaformat, MLMediaFormatKey, MLMediaFormat_Key_Sample_Rate)
+#define MLMediaFormat_Key_Sample_Rate ::MLSDK_API::MLMediaFormat_Key_Sample_RateShim()
+CREATE_GLOBAL_SHIM(ml_mediaformat, MLMediaFormatKey, MLMediaFormat_Key_Stride)
+#define MLMediaFormat_Key_Stride ::MLSDK_API::MLMediaFormat_Key_StrideShim()
+CREATE_GLOBAL_SHIM(ml_mediaformat, MLMediaFormatKey, MLMediaFormat_Key_Temporal_Layering)
+#define MLMediaFormat_Key_Temporal_Layering ::MLSDK_API::MLMediaFormat_Key_Temporal_LayeringShim()
+CREATE_GLOBAL_SHIM(ml_mediaformat, MLMediaFormatKey, MLMediaFormat_Key_Width)
+#define MLMediaFormat_Key_Width ::MLSDK_API::MLMediaFormat_Key_WidthShim()
+CREATE_GLOBAL_SHIM(ml_mediaformat, MLMediaFormatKey, MLMediaFormat_Key_Crop_Left)
+#define MLMediaFormat_Key_Crop_Left ::MLSDK_API::MLMediaFormat_Key_Crop_LeftShim()
+CREATE_GLOBAL_SHIM(ml_mediaformat, MLMediaFormatKey, MLMediaFormat_Key_Crop_Right)
+#define MLMediaFormat_Key_Crop_Right ::MLSDK_API::MLMediaFormat_Key_Crop_RightShim()
+CREATE_GLOBAL_SHIM(ml_mediaformat, MLMediaFormatKey, MLMediaFormat_Key_Crop_Bottom)
+#define MLMediaFormat_Key_Crop_Bottom ::MLSDK_API::MLMediaFormat_Key_Crop_BottomShim()
+CREATE_GLOBAL_SHIM(ml_mediaformat, MLMediaFormatKey, MLMediaFormat_Key_Crop_Top)
+#define MLMediaFormat_Key_Crop_Top ::MLSDK_API::MLMediaFormat_Key_Crop_TopShim()
+
+CREATE_FUNCTION_SHIM(ml_mediaformat, MLResult, MLMediaFormatCreateVideo)
+#define MLMediaFormatCreateVideo ::MLSDK_API::MLMediaFormatCreateVideoShim
+CREATE_FUNCTION_SHIM(ml_mediaformat, MLResult, MLMediaFormatCreateAudio)
+#define MLMediaFormatCreateAudio ::MLSDK_API::MLMediaFormatCreateAudioShim
+CREATE_FUNCTION_SHIM(ml_mediaformat, MLResult, MLMediaFormatCreateSubtitle)
+#define MLMediaFormatCreateSubtitle ::MLSDK_API::MLMediaFormatCreateSubtitleShim
+CREATE_FUNCTION_SHIM(ml_mediaformat, MLResult, MLMediaFormatDestroy)
+#define MLMediaFormatDestroy ::MLSDK_API::MLMediaFormatDestroyShim
+CREATE_FUNCTION_SHIM(ml_mediaformat, MLResult, MLMediaFormatObjectToString)
+#define MLMediaFormatObjectToString ::MLSDK_API::MLMediaFormatObjectToStringShim
+CREATE_FUNCTION_SHIM(ml_mediaformat, MLResult, MLMediaFormatGetKeyValueInt32)
+#define MLMediaFormatGetKeyValueInt32 ::MLSDK_API::MLMediaFormatGetKeyValueInt32Shim
+CREATE_FUNCTION_SHIM(ml_mediaformat, MLResult, MLMediaFormatGetKeyValueInt64)
+#define MLMediaFormatGetKeyValueInt64 ::MLSDK_API::MLMediaFormatGetKeyValueInt64Shim
+CREATE_FUNCTION_SHIM(ml_mediaformat, MLResult, MLMediaFormatGetKeyValueFloat)
+#define MLMediaFormatGetKeyValueFloat ::MLSDK_API::MLMediaFormatGetKeyValueFloatShim
+CREATE_FUNCTION_SHIM(ml_mediaformat, MLResult, MLMediaFormatGetKeySize)
+#define MLMediaFormatGetKeySize ::MLSDK_API::MLMediaFormatGetKeySizeShim
+CREATE_FUNCTION_SHIM(ml_mediaformat, MLResult, MLMediaFormatGetKeyString)
+#define MLMediaFormatGetKeyString ::MLSDK_API::MLMediaFormatGetKeyStringShim
+CREATE_FUNCTION_SHIM(ml_mediaformat, MLResult, MLMediaFormatGetKeyByteBuffer)
+#define MLMediaFormatGetKeyByteBuffer ::MLSDK_API::MLMediaFormatGetKeyByteBufferShim
+CREATE_FUNCTION_SHIM(ml_mediaformat, MLResult, MLMediaFormatKeyByteBufferRelease)
+#define MLMediaFormatKeyByteBufferRelease ::MLSDK_API::MLMediaFormatKeyByteBufferReleaseShim
+CREATE_FUNCTION_SHIM(ml_mediaformat, MLResult, MLMediaFormatSetKeyInt32)
+#define MLMediaFormatSetKeyInt32 ::MLSDK_API::MLMediaFormatSetKeyInt32Shim
+CREATE_FUNCTION_SHIM(ml_mediaformat, MLResult, MLMediaFormatSetKeyInt64)
+#define MLMediaFormatSetKeyInt64 ::MLSDK_API::MLMediaFormatSetKeyInt64Shim
+CREATE_FUNCTION_SHIM(ml_mediaformat, MLResult, MLMediaFormatSetKeyFloat)
+#define MLMediaFormatSetKeyFloat ::MLSDK_API::MLMediaFormatSetKeyFloatShim
+CREATE_FUNCTION_SHIM(ml_mediaformat, MLResult, MLMediaFormatSetKeySize)
+#define MLMediaFormatSetKeySize ::MLSDK_API::MLMediaFormatSetKeySizeShim
+CREATE_FUNCTION_SHIM(ml_mediaformat, MLResult, MLMediaFormatSetKeyString)
+#define MLMediaFormatSetKeyString ::MLSDK_API::MLMediaFormatSetKeyStringShim
+CREATE_FUNCTION_SHIM(ml_mediaformat, MLResult, MLMediaFormatSetKeyByteBuffer)
+#define MLMediaFormatSetKeyByteBuffer ::MLSDK_API::MLMediaFormatSetKeyByteBufferShim
+
+}
+
+#endif // !defined(WITH_MLSDK) || WITH_MLSDK
