Copying //UE4/Partner-MagicLeap-Staging to Dev-VR (//UE4/Dev-VR)

diff --git a/Engine/Plugins/Lumin/MLSDK/Source/MLSDK/Public/Lumin/CAPIShims/LuminAPIMediaPlayer.h b/Engine/Plugins/Lumin/MLSDK/Source/MLSDK/Public/Lumin/CAPIShims/LuminAPIMediaPlayer.h
new file mode 100644
--- /dev/null
+++ b/Engine/Plugins/Lumin/MLSDK/Source/MLSDK/Public/Lumin/CAPIShims/LuminAPIMediaPlayer.h
@@ -0,0 +1,148 @@
+// Copyright 1998-2019 Epic Games, Inc. All Rights Reserved.
+
+#pragma once
+
+#if !defined(WITH_MLSDK) || WITH_MLSDK
+
+#include "Lumin/CAPIShims/LuminAPI.h"
+
+LUMIN_THIRD_PARTY_INCLUDES_START
+#include <ml_media_common.h>
+#include <ml_media_player.h>
+LUMIN_THIRD_PARTY_INCLUDES_END
+
+namespace MLSDK_API
+{
+
+CREATE_FUNCTION_SHIM(ml_mediaplayer, MLResult, MLMediaPlayerCreate)
+#define MLMediaPlayerCreate ::MLSDK_API::MLMediaPlayerCreateShim
+CREATE_FUNCTION_SHIM(ml_mediaplayer, MLResult, MLMediaPlayerDestroy)
+#define MLMediaPlayerDestroy ::MLSDK_API::MLMediaPlayerDestroyShim
+CREATE_FUNCTION_SHIM(ml_mediaplayer, MLResult, MLMediaPlayerSetDataSourceForFD)
+#define MLMediaPlayerSetDataSourceForFD ::MLSDK_API::MLMediaPlayerSetDataSourceForFDShim
+CREATE_FUNCTION_SHIM(ml_mediaplayer, MLResult, MLMediaPlayerSetOffsetDataSourceForFD)
+#define MLMediaPlayerSetOffsetDataSourceForFD ::MLSDK_API::MLMediaPlayerSetOffsetDataSourceForFDShim
+CREATE_FUNCTION_SHIM(ml_mediaplayer, MLResult, MLMediaPlayerSetDataSourceForPath)
+#define MLMediaPlayerSetDataSourceForPath ::MLSDK_API::MLMediaPlayerSetDataSourceForPathShim
+CREATE_FUNCTION_SHIM(ml_mediaplayer, MLResult, MLMediaPlayerSetDataSourceForURI)
+#define MLMediaPlayerSetDataSourceForURI ::MLSDK_API::MLMediaPlayerSetDataSourceForURIShim
+CREATE_FUNCTION_SHIM(ml_mediaplayer, MLResult, MLMediaPlayerSetRemoteDataSourceForURI)
+#define MLMediaPlayerSetRemoteDataSourceForURI ::MLSDK_API::MLMediaPlayerSetRemoteDataSourceForURIShim
+CREATE_FUNCTION_SHIM(ml_mediaplayer, MLResult, MLMediaPlayerSetMediaDataSource)
+#define MLMediaPlayerSetMediaDataSource ::MLSDK_API::MLMediaPlayerSetMediaDataSourceShim
+CREATE_FUNCTION_SHIM(ml_mediaplayer, MLResult, MLMediaPlayerSetMediaStreamSource)
+#define MLMediaPlayerSetMediaStreamSource ::MLSDK_API::MLMediaPlayerSetMediaStreamSourceShim
+CREATE_FUNCTION_SHIM(ml_mediaplayer, MLResult, MLMediaPlayerAddTimedTextSourceForURI)
+#define MLMediaPlayerAddTimedTextSourceForURI ::MLSDK_API::MLMediaPlayerAddTimedTextSourceForURIShim
+CREATE_FUNCTION_SHIM(ml_mediaplayer, MLResult, MLMediaPlayerAddTimedTextSourceForFD)
+#define MLMediaPlayerAddTimedTextSourceForFD ::MLSDK_API::MLMediaPlayerAddTimedTextSourceForFDShim
+CREATE_FUNCTION_SHIM(ml_mediaplayer, MLResult, MLMediaPlayerSetNextPlayer)
+#define MLMediaPlayerSetNextPlayer ::MLSDK_API::MLMediaPlayerSetNextPlayerShim
+CREATE_FUNCTION_SHIM(ml_mediaplayer, MLResult, MLMediaPlayerGetTrackCount)
+#define MLMediaPlayerGetTrackCount ::MLSDK_API::MLMediaPlayerGetTrackCountShim
+CREATE_FUNCTION_SHIM(ml_mediaplayer, MLResult, MLMediaPlayerGetTrackType)
+#define MLMediaPlayerGetTrackType ::MLSDK_API::MLMediaPlayerGetTrackTypeShim
+CREATE_FUNCTION_SHIM(ml_mediaplayer, MLResult, MLMediaPlayerGetTrackLanguage)
+#define MLMediaPlayerGetTrackLanguage ::MLSDK_API::MLMediaPlayerGetTrackLanguageShim
+CREATE_FUNCTION_SHIM(ml_mediaplayer, MLResult, MLMediaPlayerGetTrackMediaFormat)
+#define MLMediaPlayerGetTrackMediaFormat ::MLSDK_API::MLMediaPlayerGetTrackMediaFormatShim
+CREATE_FUNCTION_SHIM(ml_mediaplayer, MLResult, MLMediaPlayerSelectTrack)
+#define MLMediaPlayerSelectTrack ::MLSDK_API::MLMediaPlayerSelectTrackShim
+CREATE_FUNCTION_SHIM(ml_mediaplayer, MLResult, MLMediaPlayerUnselectTrack)
+#define MLMediaPlayerUnselectTrack ::MLSDK_API::MLMediaPlayerUnselectTrackShim
+CREATE_FUNCTION_SHIM(ml_mediaplayer, MLResult, MLMediaPlayerPrepare)
+#define MLMediaPlayerPrepare ::MLSDK_API::MLMediaPlayerPrepareShim
+CREATE_FUNCTION_SHIM(ml_mediaplayer, MLResult, MLMediaPlayerPrepareAsync)
+#define MLMediaPlayerPrepareAsync ::MLSDK_API::MLMediaPlayerPrepareAsyncShim
+CREATE_FUNCTION_SHIM(ml_mediaplayer, MLResult, MLMediaPlayerStart)
+#define MLMediaPlayerStart ::MLSDK_API::MLMediaPlayerStartShim
+CREATE_FUNCTION_SHIM(ml_mediaplayer, MLResult, MLMediaPlayerStop)
+#define MLMediaPlayerStop ::MLSDK_API::MLMediaPlayerStopShim
+CREATE_FUNCTION_SHIM(ml_mediaplayer, MLResult, MLMediaPlayerPause)
+#define MLMediaPlayerPause ::MLSDK_API::MLMediaPlayerPauseShim
+CREATE_FUNCTION_SHIM(ml_mediaplayer, MLResult, MLMediaPlayerReset)
+#define MLMediaPlayerReset ::MLSDK_API::MLMediaPlayerResetShim
+CREATE_FUNCTION_SHIM(ml_mediaplayer, MLResult, MLMediaPlayerSeekTo)
+#define MLMediaPlayerSeekTo ::MLSDK_API::MLMediaPlayerSeekToShim
+CREATE_FUNCTION_SHIM(ml_mediaplayer, MLResult, MLMediaPlayerSetOnBufferAvailableCallback)
+#define MLMediaPlayerSetOnBufferAvailableCallback ::MLSDK_API::MLMediaPlayerSetOnBufferAvailableCallbackShim
+CREATE_FUNCTION_SHIM(ml_mediaplayer, MLResult, MLMediaPlayerAcquireNextAvailableBuffer)
+#define MLMediaPlayerAcquireNextAvailableBuffer ::MLSDK_API::MLMediaPlayerAcquireNextAvailableBufferShim
+CREATE_FUNCTION_SHIM(ml_mediaplayer, MLResult, MLMediaPlayerReleaseBuffer)
+#define MLMediaPlayerReleaseBuffer ::MLSDK_API::MLMediaPlayerReleaseBufferShim
+CREATE_FUNCTION_SHIM(ml_mediaplayer, MLResult, MLMediaPlayerGetFrameTransformationMatrix)
+#define MLMediaPlayerGetFrameTransformationMatrix ::MLSDK_API::MLMediaPlayerGetFrameTransformationMatrixShim
+CREATE_FUNCTION_SHIM(ml_mediaplayer, MLResult, MLMediaPlayerGetFrameTimestamp)
+#define MLMediaPlayerGetFrameTimestamp ::MLSDK_API::MLMediaPlayerGetFrameTimestampShim
+CREATE_FUNCTION_SHIM(ml_mediaplayer, MLResult, MLMediaPlayerGetFrameQueueBufferTimestamp)
+#define MLMediaPlayerGetFrameQueueBufferTimestamp ::MLSDK_API::MLMediaPlayerGetFrameQueueBufferTimestampShim
+CREATE_FUNCTION_SHIM(ml_mediaplayer, MLResult, MLMediaPlayerGetFrameNumber)
+#define MLMediaPlayerGetFrameNumber ::MLSDK_API::MLMediaPlayerGetFrameNumberShim
+CREATE_DEPRECATED_MSG_SHIM(ml_mediaplayer, MLResult, MLMediaPlayerSetEventCallbacks, "Replaced by MLMediaPlayerSetEventCallbacksEx.")
+#define MLMediaPlayerSetEventCallbacks ::MLSDK_API::MLMediaPlayerSetEventCallbacksShim
+CREATE_FUNCTION_SHIM(ml_mediaplayer, MLResult, MLMediaPlayerSetEventCallbacksEx)
+#define MLMediaPlayerSetEventCallbacksEx ::MLSDK_API::MLMediaPlayerSetEventCallbacksExShim
+CREATE_FUNCTION_SHIM(ml_mediaplayer, MLResult, MLMediaPlayerGetVideoSize)
+#define MLMediaPlayerGetVideoSize ::MLSDK_API::MLMediaPlayerGetVideoSizeShim
+CREATE_FUNCTION_SHIM(ml_mediaplayer, MLResult, MLMediaPlayerGetCurrentPosition)
+#define MLMediaPlayerGetCurrentPosition ::MLSDK_API::MLMediaPlayerGetCurrentPositionShim
+CREATE_FUNCTION_SHIM(ml_mediaplayer, MLResult, MLMediaPlayerGetDuration)
+#define MLMediaPlayerGetDuration ::MLSDK_API::MLMediaPlayerGetDurationShim
+CREATE_FUNCTION_SHIM(ml_mediaplayer, MLResult, MLMediaPlayerGetBufferingUpdate)
+#define MLMediaPlayerGetBufferingUpdate ::MLSDK_API::MLMediaPlayerGetBufferingUpdateShim
+CREATE_FUNCTION_SHIM(ml_mediaplayer, MLResult, MLMediaPlayerGetInfo)
+#define MLMediaPlayerGetInfo ::MLSDK_API::MLMediaPlayerGetInfoShim
+CREATE_FUNCTION_SHIM(ml_mediaplayer, MLResult, MLMediaPlayerGetError)
+#define MLMediaPlayerGetError ::MLSDK_API::MLMediaPlayerGetErrorShim
+CREATE_FUNCTION_SHIM(ml_mediaplayer, MLResult, MLMediaPlayerGetTimedText)
+#define MLMediaPlayerGetTimedText ::MLSDK_API::MLMediaPlayerGetTimedTextShim
+CREATE_DEPRECATED_MSG_SHIM(ml_mediaplayer, MLResult, MLMediaPlayerGetSubtitle, "Replaced by MLMediaPlayerGetSubtitleEx")
+#define MLMediaPlayerGetSubtitle ::MLSDK_API::MLMediaPlayerGetSubtitleShim
+CREATE_FUNCTION_SHIM(ml_mediaplayer, MLResult, MLMediaPlayerGetSubtitleEx)
+#define MLMediaPlayerGetSubtitleEx ::MLSDK_API::MLMediaPlayerGetSubtitleExShim
+CREATE_FUNCTION_SHIM(ml_mediaplayer, MLResult, MLMediaPlayerReleaseSubtitleEx)
+#define MLMediaPlayerReleaseSubtitleEx ::MLSDK_API::MLMediaPlayerReleaseSubtitleExShim
+CREATE_FUNCTION_SHIM(ml_mediaplayer, MLResult, MLMediaPlayerSetOnMediaSubtitleUpdateCallback)
+#define MLMediaPlayerSetOnMediaSubtitleUpdateCallback ::MLSDK_API::MLMediaPlayerSetOnMediaSubtitleUpdateCallbackShim
+CREATE_FUNCTION_SHIM(ml_mediaplayer, MLResult, MLMediaPlayerGetMetadata)
+#define MLMediaPlayerGetMetadata ::MLSDK_API::MLMediaPlayerGetMetadataShim
+CREATE_FUNCTION_SHIM(ml_mediaplayer, MLResult, MLMediaPlayerSetVideoScalingMode)
+#define MLMediaPlayerSetVideoScalingMode ::MLSDK_API::MLMediaPlayerSetVideoScalingModeShim
+CREATE_FUNCTION_SHIM(ml_mediaplayer, MLResult, MLMediaPlayerSetLooping)
+#define MLMediaPlayerSetLooping ::MLSDK_API::MLMediaPlayerSetLoopingShim
+CREATE_FUNCTION_SHIM(ml_mediaplayer, MLResult, MLMediaPlayerSetVolume)
+#define MLMediaPlayerSetVolume ::MLSDK_API::MLMediaPlayerSetVolumeShim
+CREATE_FUNCTION_SHIM(ml_mediaplayer, MLResult, MLMediaPlayerGetAudioHandle)
+#define MLMediaPlayerGetAudioHandle ::MLSDK_API::MLMediaPlayerGetAudioHandleShim
+CREATE_FUNCTION_SHIM(ml_mediaplayer, MLResult, MLMediaPlayerPollStates)
+#define MLMediaPlayerPollStates ::MLSDK_API::MLMediaPlayerPollStatesShim
+CREATE_DEPRECATED_MSG_SHIM(ml_mediaplayer, MLResult, MLMediaPlayerGetDRMInfo, "Replaced by MLMediaPlayerGetTrackDRMInfo.")
+#define MLMediaPlayerGetDRMInfo ::MLSDK_API::MLMediaPlayerGetDRMInfoShim
+CREATE_DEPRECATED_MSG_SHIM(ml_mediaplayer, MLResult, MLMediaPlayerPrepareDRM, "Replaced by MLMediaPlayerPrepareTrackDRM.")
+#define MLMediaPlayerPrepareDRM ::MLSDK_API::MLMediaPlayerPrepareDRMShim
+CREATE_FUNCTION_SHIM(ml_mediaplayer, MLResult, MLMediaPlayerGetTrackDRMInfo)
+#define MLMediaPlayerGetTrackDRMInfo ::MLSDK_API::MLMediaPlayerGetTrackDRMInfoShim
+CREATE_FUNCTION_SHIM(ml_mediaplayer, MLResult, MLMediaPlayerPrepareTrackDRM)
+#define MLMediaPlayerPrepareTrackDRM ::MLSDK_API::MLMediaPlayerPrepareTrackDRMShim
+CREATE_FUNCTION_SHIM(ml_mediaplayer, MLResult, MLMediaPlayerReleaseDRM)
+#define MLMediaPlayerReleaseDRM ::MLSDK_API::MLMediaPlayerReleaseDRMShim
+CREATE_FUNCTION_SHIM(ml_mediaplayer, MLResult, MLMediaPlayerGetDefaultBufferingSettings)
+#define MLMediaPlayerGetDefaultBufferingSettings ::MLSDK_API::MLMediaPlayerGetDefaultBufferingSettingsShim
+CREATE_FUNCTION_SHIM(ml_mediaplayer, MLResult, MLMediaPlayerGetBufferingSettings)
+#define MLMediaPlayerGetBufferingSettings ::MLSDK_API::MLMediaPlayerGetBufferingSettingsShim
+CREATE_FUNCTION_SHIM(ml_mediaplayer, MLResult, MLMediaPlayerSetBufferingSettings)
+#define MLMediaPlayerSetBufferingSettings ::MLSDK_API::MLMediaPlayerSetBufferingSettingsShim
+CREATE_FUNCTION_SHIM(ml_mediaplayer, MLResult, MLMediaStreamSourceCreate)
+#define MLMediaStreamSourceCreate ::MLSDK_API::MLMediaStreamSourceCreateShim
+CREATE_FUNCTION_SHIM(ml_mediaplayer, MLResult, MLMediaStreamSourceDestroy)
+#define MLMediaStreamSourceDestroy ::MLSDK_API::MLMediaStreamSourceDestroyShim
+CREATE_FUNCTION_SHIM(ml_mediaplayer, MLResult, MLMediaStreamSourceGetBuffer)
+#define MLMediaStreamSourceGetBuffer ::MLSDK_API::MLMediaStreamSourceGetBufferShim
+CREATE_FUNCTION_SHIM(ml_mediaplayer, MLResult, MLMediaStreamSourcePushBuffer)
+#define MLMediaStreamSourcePushBuffer ::MLSDK_API::MLMediaStreamSourcePushBufferShim
+CREATE_FUNCTION_SHIM(ml_mediaplayer, MLResult, MLMediaStreamSourcePushEOS)
+#define MLMediaStreamSourcePushEOS ::MLSDK_API::MLMediaStreamSourcePushEOSShim
+
+}
+
+#endif // !defined(WITH_MLSDK) || WITH_MLSDK
diff --git a/Engine/Plugins/Lumin/MLSDK/Source/MLSDK/Public/Lumin/CAPIShims/LuminAPIMediaStreamSource.h b/Engine/Plugins/Lumin/MLSDK/Source/MLSDK/Public/Lumin/CAPIShims/LuminAPIMediaStreamSource.h
new file mode 100644
--- /dev/null
+++ b/Engine/Plugins/Lumin/MLSDK/Source/MLSDK/Public/Lumin/CAPIShims/LuminAPIMediaStreamSource.h
@@ -0,0 +1,13 @@
+// Copyright 1998-2019 Epic Games, Inc. All Rights Reserved.
+
+#pragma once
+
+#if !defined(WITH_MLSDK) || WITH_MLSDK
+
+#include "Lumin/CAPIShims/LuminAPI.h"
+
+LUMIN_THIRD_PARTY_INCLUDES_START
+#include <ml_media_stream_source.h>
+LUMIN_THIRD_PARTY_INCLUDES_END
+
+#endif // !defined(WITH_MLSDK) || WITH_MLSDK
diff --git a/Engine/Plugins/Lumin/MLSDK/Source/MLSDK/Public/Lumin/CAPIShims/LuminAPIMediaTypes.h b/Engine/Plugins/Lumin/MLSDK/Source/MLSDK/Public/Lumin/CAPIShims/LuminAPIMediaTypes.h
new file mode 100644
--- /dev/null
+++ b/Engine/Plugins/Lumin/MLSDK/Source/MLSDK/Public/Lumin/CAPIShims/LuminAPIMediaTypes.h
@@ -0,0 +1,13 @@
+// Copyright 1998-2019 Epic Games, Inc. All Rights Reserved.
+
+#pragma once
+
+#if !defined(WITH_MLSDK) || WITH_MLSDK
+
+#include "Lumin/CAPIShims/LuminAPI.h"
+
+LUMIN_THIRD_PARTY_INCLUDES_START
+#include <ml_media_types.h>
+LUMIN_THIRD_PARTY_INCLUDES_END
+
+#endif // !defined(WITH_MLSDK) || WITH_MLSDK
diff --git a/Engine/Plugins/Lumin/MLSDK/Source/MLSDK/Public/Lumin/CAPIShims/LuminAPIMeshing.h b/Engine/Plugins/Lumin/MLSDK/Source/MLSDK/Public/Lumin/CAPIShims/LuminAPIMeshing.h
new file mode 100644
--- /dev/null
+++ b/Engine/Plugins/Lumin/MLSDK/Source/MLSDK/Public/Lumin/CAPIShims/LuminAPIMeshing.h
@@ -0,0 +1,37 @@
+// Copyright 1998-2019 Epic Games, Inc. All Rights Reserved.
+
+#pragma once
+
+#if !defined(WITH_MLSDK) || WITH_MLSDK
+
+#include "Lumin/CAPIShims/LuminAPI.h"
+
+LUMIN_THIRD_PARTY_INCLUDES_START
+#include <ml_meshing2.h>
+LUMIN_THIRD_PARTY_INCLUDES_END
+
+namespace MLSDK_API
+{
+
+CREATE_FUNCTION_SHIM(ml_perception_client, MLResult, MLMeshingCreateClient)
+#define MLMeshingCreateClient ::MLSDK_API::MLMeshingCreateClientShim
+CREATE_FUNCTION_SHIM(ml_perception_client, MLResult, MLMeshingDestroyClient)
+#define MLMeshingDestroyClient ::MLSDK_API::MLMeshingDestroyClientShim
+CREATE_FUNCTION_SHIM(ml_perception_client, MLResult, MLMeshingInitSettings)
+#define MLMeshingInitSettings ::MLSDK_API::MLMeshingInitSettingsShim
+CREATE_FUNCTION_SHIM(ml_perception_client, MLResult, MLMeshingUpdateSettings)
+#define MLMeshingUpdateSettings ::MLSDK_API::MLMeshingUpdateSettingsShim
+CREATE_FUNCTION_SHIM(ml_perception_client, MLResult, MLMeshingRequestMeshInfo)
+#define MLMeshingRequestMeshInfo ::MLSDK_API::MLMeshingRequestMeshInfoShim
+CREATE_FUNCTION_SHIM(ml_perception_client, MLResult, MLMeshingGetMeshInfoResult)
+#define MLMeshingGetMeshInfoResult ::MLSDK_API::MLMeshingGetMeshInfoResultShim
+CREATE_FUNCTION_SHIM(ml_perception_client, MLResult, MLMeshingRequestMesh)
+#define MLMeshingRequestMesh ::MLSDK_API::MLMeshingRequestMeshShim
+CREATE_FUNCTION_SHIM(ml_perception_client, MLResult, MLMeshingGetMeshResult)
+#define MLMeshingGetMeshResult ::MLSDK_API::MLMeshingGetMeshResultShim
+CREATE_FUNCTION_SHIM(ml_perception_client, MLResult, MLMeshingFreeResource)
+#define MLMeshingFreeResource ::MLSDK_API::MLMeshingFreeResourceShim
+
+}
+
+#endif // !defined(WITH_MLSDK) || WITH_MLSDK
diff --git a/Engine/Plugins/Lumin/MLSDK/Source/MLSDK/Public/Lumin/CAPIShims/LuminAPIMusicService.h b/Engine/Plugins/Lumin/MLSDK/Source/MLSDK/Public/Lumin/CAPIShims/LuminAPIMusicService.h
new file mode 100644
--- /dev/null
+++ b/Engine/Plugins/Lumin/MLSDK/Source/MLSDK/Public/Lumin/CAPIShims/LuminAPIMusicService.h
@@ -0,0 +1,74 @@
+// Copyright 1998-2019 Epic Games, Inc. All Rights Reserved.
+
+#pragma once
+
+#if !defined(WITH_MLSDK) || WITH_MLSDK
+
+#include "Lumin/CAPIShims/LuminAPI.h"
+
+LUMIN_THIRD_PARTY_INCLUDES_START
+#include <ml_music_service_common.h>
+#include <ml_music_service.h>
+LUMIN_THIRD_PARTY_INCLUDES_END
+
+namespace MLSDK_API
+{
+
+CREATE_FUNCTION_SHIM(ml_musicservice, MLResult, MLMusicServiceConnect)
+#define MLMusicServiceConnect ::MLSDK_API::MLMusicServiceConnectShim
+CREATE_FUNCTION_SHIM(ml_musicservice, MLResult, MLMusicServiceDisconnect)
+#define MLMusicServiceDisconnect ::MLSDK_API::MLMusicServiceDisconnectShim
+CREATE_FUNCTION_SHIM(ml_musicservice, MLResult, MLMusicServiceSetCallbacks)
+#define MLMusicServiceSetCallbacks ::MLSDK_API::MLMusicServiceSetCallbacksShim
+CREATE_FUNCTION_SHIM(ml_musicservice, MLResult, MLMusicServiceSetAuthString)
+#define MLMusicServiceSetAuthString ::MLSDK_API::MLMusicServiceSetAuthStringShim
+CREATE_FUNCTION_SHIM(ml_musicservice, MLResult, MLMusicServiceSetURL)
+#define MLMusicServiceSetURL ::MLSDK_API::MLMusicServiceSetURLShim
+CREATE_FUNCTION_SHIM(ml_musicservice, MLResult, MLMusicServiceSetPlayList)
+#define MLMusicServiceSetPlayList ::MLSDK_API::MLMusicServiceSetPlayListShim
+CREATE_FUNCTION_SHIM(ml_musicservice, MLResult, MLMusicServiceStart)
+#define MLMusicServiceStart ::MLSDK_API::MLMusicServiceStartShim
+CREATE_FUNCTION_SHIM(ml_musicservice, MLResult, MLMusicServiceStop)
+#define MLMusicServiceStop ::MLSDK_API::MLMusicServiceStopShim
+CREATE_FUNCTION_SHIM(ml_musicservice, MLResult, MLMusicServicePause)
+#define MLMusicServicePause ::MLSDK_API::MLMusicServicePauseShim
+CREATE_FUNCTION_SHIM(ml_musicservice, MLResult, MLMusicServiceResume)
+#define MLMusicServiceResume ::MLSDK_API::MLMusicServiceResumeShim
+CREATE_FUNCTION_SHIM(ml_musicservice, MLResult, MLMusicServiceSeek)
+#define MLMusicServiceSeek ::MLSDK_API::MLMusicServiceSeekShim
+CREATE_FUNCTION_SHIM(ml_musicservice, MLResult, MLMusicServiceNext)
+#define MLMusicServiceNext ::MLSDK_API::MLMusicServiceNextShim
+CREATE_FUNCTION_SHIM(ml_musicservice, MLResult, MLMusicServicePrevious)
+#define MLMusicServicePrevious ::MLSDK_API::MLMusicServicePreviousShim
+CREATE_FUNCTION_SHIM(ml_musicservice, MLResult, MLMusicServiceSetShuffle)
+#define MLMusicServiceSetShuffle ::MLSDK_API::MLMusicServiceSetShuffleShim
+CREATE_FUNCTION_SHIM(ml_musicservice, MLResult, MLMusicServiceSetRepeat)
+#define MLMusicServiceSetRepeat ::MLSDK_API::MLMusicServiceSetRepeatShim
+CREATE_FUNCTION_SHIM(ml_musicservice, MLResult, MLMusicServiceSetVolume)
+#define MLMusicServiceSetVolume ::MLSDK_API::MLMusicServiceSetVolumeShim
+CREATE_FUNCTION_SHIM(ml_musicservice, MLResult, MLMusicServiceGetTrackLength)
+#define MLMusicServiceGetTrackLength ::MLSDK_API::MLMusicServiceGetTrackLengthShim
+CREATE_FUNCTION_SHIM(ml_musicservice, MLResult, MLMusicServiceGetCurrentPosition)
+#define MLMusicServiceGetCurrentPosition ::MLSDK_API::MLMusicServiceGetCurrentPositionShim
+CREATE_FUNCTION_SHIM(ml_musicservice, MLResult, MLMusicServiceGetStatus)
+#define MLMusicServiceGetStatus ::MLSDK_API::MLMusicServiceGetStatusShim
+CREATE_FUNCTION_SHIM(ml_musicservice, MLResult, MLMusicServiceGetError)
+#define MLMusicServiceGetError ::MLSDK_API::MLMusicServiceGetErrorShim
+CREATE_FUNCTION_SHIM(ml_musicservice, MLResult, MLMusicServiceGetPlaybackState)
+#define MLMusicServiceGetPlaybackState ::MLSDK_API::MLMusicServiceGetPlaybackStateShim
+CREATE_FUNCTION_SHIM(ml_musicservice, MLResult, MLMusicServiceGetRepeatState)
+#define MLMusicServiceGetRepeatState ::MLSDK_API::MLMusicServiceGetRepeatStateShim
+CREATE_FUNCTION_SHIM(ml_musicservice, MLResult, MLMusicServiceGetShuffleState)
+#define MLMusicServiceGetShuffleState ::MLSDK_API::MLMusicServiceGetShuffleStateShim
+CREATE_FUNCTION_SHIM(ml_musicservice, MLResult, MLMusicServiceGetVolume)
+#define MLMusicServiceGetVolume ::MLSDK_API::MLMusicServiceGetVolumeShim
+CREATE_FUNCTION_SHIM(ml_musicservice, MLResult, MLMusicServiceGetMetadata)
+#define MLMusicServiceGetMetadata ::MLSDK_API::MLMusicServiceGetMetadataShim
+CREATE_FUNCTION_SHIM(ml_musicservice, MLResult, MLMusicServiceGetMetadataEx)
+#define MLMusicServiceGetMetadataEx ::MLSDK_API::MLMusicServiceGetMetadataExShim
+CREATE_FUNCTION_SHIM(ml_musicservice, MLResult, MLMusicServiceReleaseMetadata)
+#define MLMusicServiceReleaseMetadata ::MLSDK_API::MLMusicServiceReleaseMetadataShim
+
+}
+
+#endif // !defined(WITH_MLSDK) || WITH_MLSDK
diff --git a/Engine/Plugins/Lumin/MLSDK/Source/MLSDK/Public/Lumin/CAPIShims/LuminAPIMusicServiceProvider.h b/Engine/Plugins/Lumin/MLSDK/Source/MLSDK/Public/Lumin/CAPIShims/LuminAPIMusicServiceProvider.h
new file mode 100644
--- /dev/null
+++ b/Engine/Plugins/Lumin/MLSDK/Source/MLSDK/Public/Lumin/CAPIShims/LuminAPIMusicServiceProvider.h
@@ -0,0 +1,52 @@
+// Copyright 1998-2019 Epic Games, Inc. All Rights Reserved.
+
+#pragma once
+
+#if !defined(WITH_MLSDK) || WITH_MLSDK
+
+#include "Lumin/CAPIShims/LuminAPI.h"
+
+LUMIN_THIRD_PARTY_INCLUDES_START
+#include <ml_music_service_common.h>
+#include <ml_music_service_provider.h>
+LUMIN_THIRD_PARTY_INCLUDES_END
+
+namespace MLSDK_API
+{
+
+CREATE_FUNCTION_SHIM(ml_musicservice_provider, MLResult, MLMusicServiceProviderCreate)
+#define MLMusicServiceProviderCreate ::MLSDK_API::MLMusicServiceProviderCreateShim
+CREATE_FUNCTION_SHIM(ml_musicservice_provider, MLResult, MLMusicServiceProviderStart)
+#define MLMusicServiceProviderStart ::MLSDK_API::MLMusicServiceProviderStartShim
+CREATE_FUNCTION_SHIM(ml_musicservice_provider, MLResult, MLMusicServiceProviderSetAudioOutput)
+#define MLMusicServiceProviderSetAudioOutput ::MLSDK_API::MLMusicServiceProviderSetAudioOutputShim
+CREATE_FUNCTION_SHIM(ml_musicservice_provider, MLResult, MLMusicServiceProviderWriteAudioOutput)
+#define MLMusicServiceProviderWriteAudioOutput ::MLSDK_API::MLMusicServiceProviderWriteAudioOutputShim
+CREATE_FUNCTION_SHIM(ml_musicservice_provider, MLResult, MLMusicServiceProviderFlushAudioOutput)
+#define MLMusicServiceProviderFlushAudioOutput ::MLSDK_API::MLMusicServiceProviderFlushAudioOutputShim
+CREATE_FUNCTION_SHIM(ml_musicservice_provider, MLResult, MLMusicServiceProviderReleaseAudioOutput)
+#define MLMusicServiceProviderReleaseAudioOutput ::MLSDK_API::MLMusicServiceProviderReleaseAudioOutputShim
+CREATE_FUNCTION_SHIM(ml_musicservice_provider, MLResult, MLMusicServiceProviderSetVolume)
+#define MLMusicServiceProviderSetVolume ::MLSDK_API::MLMusicServiceProviderSetVolumeShim
+CREATE_FUNCTION_SHIM(ml_musicservice_provider, MLResult, MLMusicServiceProviderNotifyPositionChange)
+#define MLMusicServiceProviderNotifyPositionChange ::MLSDK_API::MLMusicServiceProviderNotifyPositionChangeShim
+CREATE_FUNCTION_SHIM(ml_musicservice_provider, MLResult, MLMusicServiceProviderNotifyPlaybackStateChange)
+#define MLMusicServiceProviderNotifyPlaybackStateChange ::MLSDK_API::MLMusicServiceProviderNotifyPlaybackStateChangeShim
+CREATE_FUNCTION_SHIM(ml_musicservice_provider, MLResult, MLMusicServiceProviderNotifyRepeatStateChange)
+#define MLMusicServiceProviderNotifyRepeatStateChange ::MLSDK_API::MLMusicServiceProviderNotifyRepeatStateChangeShim
+CREATE_FUNCTION_SHIM(ml_musicservice_provider, MLResult, MLMusicServiceProviderNotifyShuffleStateChange)
+#define MLMusicServiceProviderNotifyShuffleStateChange ::MLSDK_API::MLMusicServiceProviderNotifyShuffleStateChangeShim
+CREATE_FUNCTION_SHIM(ml_musicservice_provider, MLResult, MLMusicServiceProviderNotifyStatus)
+#define MLMusicServiceProviderNotifyStatus ::MLSDK_API::MLMusicServiceProviderNotifyStatusShim
+CREATE_FUNCTION_SHIM(ml_musicservice_provider, MLResult, MLMusicServiceProviderNotifyError)
+#define MLMusicServiceProviderNotifyError ::MLSDK_API::MLMusicServiceProviderNotifyErrorShim
+CREATE_FUNCTION_SHIM(ml_musicservice_provider, MLResult, MLMusicServiceProviderNotifyMetadataChange)
+#define MLMusicServiceProviderNotifyMetadataChange ::MLSDK_API::MLMusicServiceProviderNotifyMetadataChangeShim
+CREATE_FUNCTION_SHIM(ml_musicservice_provider, MLResult, MLMusicServiceProviderNotifyVolumeChange)
+#define MLMusicServiceProviderNotifyVolumeChange ::MLSDK_API::MLMusicServiceProviderNotifyVolumeChangeShim
+CREATE_FUNCTION_SHIM(ml_musicservice_provider, MLResult, MLMusicServiceProviderTerminate)
+#define MLMusicServiceProviderTerminate ::MLSDK_API::MLMusicServiceProviderTerminateShim
+
+}
+
+#endif // !defined(WITH_MLSDK) || WITH_MLSDK
diff --git a/Engine/Plugins/Lumin/MLSDK/Source/MLSDK/Public/Lumin/CAPIShims/LuminAPINetworking.h b/Engine/Plugins/Lumin/MLSDK/Source/MLSDK/Public/Lumin/CAPIShims/LuminAPINetworking.h
new file mode 100644
--- /dev/null
+++ b/Engine/Plugins/Lumin/MLSDK/Source/MLSDK/Public/Lumin/CAPIShims/LuminAPINetworking.h
@@ -0,0 +1,24 @@
+// Copyright 1998-2019 Epic Games, Inc. All Rights Reserved.
+#pragma once
+
+#if !defined(WITH_MLSDK) || WITH_MLSDK
+
+#include "Lumin/CAPIShims/LuminAPI.h"
+
+LUMIN_THIRD_PARTY_INCLUDES_START
+#include <ml_networking.h>
+LUMIN_THIRD_PARTY_INCLUDES_END
+
+namespace MLSDK_API
+{
+
+CREATE_FUNCTION_SHIM(ml_networking, MLResult, MLNetworkingIsInternetConnected)
+#define MLNetworkingIsInternetConnected ::MLSDK_API::MLNetworkingIsInternetConnectedShim
+CREATE_FUNCTION_SHIM(ml_networking, MLResult, MLNetworkingGetWiFiData)
+#define MLNetworkingGetWiFiData ::MLSDK_API::MLNetworkingGetWiFiDataShim
+CREATE_FUNCTION_SHIM(ml_networking, const char*, MLNetworkingGetResultString)
+#define MLNetworkingGetResultString ::MLSDK_API::MLNetworkingGetResultStringShim
+
+}
+
+#endif // !defined(WITH_MLSDK) || WITH_MLSDK
diff --git a/Engine/Plugins/Lumin/MLSDK/Source/MLSDK/Public/Lumin/CAPIShims/LuminAPIPassableWorld.h b/Engine/Plugins/Lumin/MLSDK/Source/MLSDK/Public/Lumin/CAPIShims/LuminAPIPassableWorld.h
new file mode 100644
--- /dev/null
+++ b/Engine/Plugins/Lumin/MLSDK/Source/MLSDK/Public/Lumin/CAPIShims/LuminAPIPassableWorld.h
@@ -0,0 +1,13 @@
+// Copyright 1998-2019 Epic Games, Inc. All Rights Reserved.
+
+#pragma once
+
+#if !defined(WITH_MLSDK) || WITH_MLSDK
+
+#include "Lumin/CAPIShims/LuminAPI.h"
+
+LUMIN_THIRD_PARTY_INCLUDES_START
+#include <ml_passable_world.h>
+LUMIN_THIRD_PARTY_INCLUDES_END
+
+#endif // !defined(WITH_MLSDK) || WITH_MLSDK
diff --git a/Engine/Plugins/Lumin/MLSDK/Source/MLSDK/Public/Lumin/CAPIShims/LuminAPIPerception.h b/Engine/Plugins/Lumin/MLSDK/Source/MLSDK/Public/Lumin/CAPIShims/LuminAPIPerception.h
new file mode 100644
--- /dev/null
+++ b/Engine/Plugins/Lumin/MLSDK/Source/MLSDK/Public/Lumin/CAPIShims/LuminAPIPerception.h
@@ -0,0 +1,29 @@
+// Copyright 1998-2019 Epic Games, Inc. All Rights Reserved.
+
+#pragma once
+
+#if !defined(WITH_MLSDK) || WITH_MLSDK
+
+#include "Lumin/CAPIShims/LuminAPI.h"
+
+LUMIN_THIRD_PARTY_INCLUDES_START
+#include <ml_perception.h>
+LUMIN_THIRD_PARTY_INCLUDES_END
+
+namespace MLSDK_API
+{
+
+CREATE_FUNCTION_SHIM(ml_perception_client, MLResult, MLPerceptionInitSettings)
+#define MLPerceptionInitSettings ::MLSDK_API::MLPerceptionInitSettingsShim
+CREATE_FUNCTION_SHIM(ml_perception_client, MLResult, MLPerceptionStartup)
+#define MLPerceptionStartup ::MLSDK_API::MLPerceptionStartupShim
+CREATE_FUNCTION_SHIM(ml_perception_client, MLResult, MLPerceptionShutdown)
+#define MLPerceptionShutdown ::MLSDK_API::MLPerceptionShutdownShim
+CREATE_FUNCTION_SHIM(ml_perception_client, MLResult, MLPerceptionGetSnapshot)
+#define MLPerceptionGetSnapshot ::MLSDK_API::MLPerceptionGetSnapshotShim
+CREATE_FUNCTION_SHIM(ml_perception_client, MLResult, MLPerceptionReleaseSnapshot)
+#define MLPerceptionReleaseSnapshot ::MLSDK_API::MLPerceptionReleaseSnapshotShim
+
+}
+
+#endif // !defined(WITH_MLSDK) || WITH_MLSDK
diff --git a/Engine/Plugins/Lumin/MLSDK/Source/MLSDK/Public/Lumin/CAPIShims/LuminAPIPersistentCoordinateFrames.h b/Engine/Plugins/Lumin/MLSDK/Source/MLSDK/Public/Lumin/CAPIShims/LuminAPIPersistentCoordinateFrames.h
new file mode 100644
--- /dev/null
+++ b/Engine/Plugins/Lumin/MLSDK/Source/MLSDK/Public/Lumin/CAPIShims/LuminAPIPersistentCoordinateFrames.h
@@ -0,0 +1,35 @@
+// Copyright 1998-2019 Epic Games, Inc. All Rights Reserved.
+
+#pragma once
+
+#if !defined(WITH_MLSDK) || WITH_MLSDK
+
+#include "Lumin/CAPIShims/LuminAPI.h"
+
+LUMIN_THIRD_PARTY_INCLUDES_START
+#include <ml_persistent_coordinate_frames.h>
+LUMIN_THIRD_PARTY_INCLUDES_END
+
+namespace MLSDK_API
+{
+
+CREATE_FUNCTION_SHIM(ml_perception_client, MLResult, MLPersistentCoordinateFrameTrackerCreate)
+#define MLPersistentCoordinateFrameTrackerCreate ::MLSDK_API::MLPersistentCoordinateFrameTrackerCreateShim
+CREATE_FUNCTION_SHIM(ml_perception_client, MLResult, MLPersistentCoordinateFrameGetCount)
+#define MLPersistentCoordinateFrameGetCount ::MLSDK_API::MLPersistentCoordinateFrameGetCountShim
+CREATE_FUNCTION_SHIM(ml_perception_client, MLResult, MLPersistentCoordinateFrameGetAll)
+#define MLPersistentCoordinateFrameGetAll ::MLSDK_API::MLPersistentCoordinateFrameGetAllShim
+CREATE_FUNCTION_SHIM(ml_perception_client, MLResult, MLPersistentCoordinateFrameGetAllEx)
+#define MLPersistentCoordinateFrameGetAllEx ::MLSDK_API::MLPersistentCoordinateFrameGetAllExShim
+CREATE_FUNCTION_SHIM(ml_perception_client, MLResult, MLPersistentCoordinateFrameGetClosest)
+#define MLPersistentCoordinateFrameGetClosest ::MLSDK_API::MLPersistentCoordinateFrameGetClosestShim
+CREATE_FUNCTION_SHIM(ml_perception_client, MLResult, MLPersistentCoordinateFrameTrackerDestroy)
+#define MLPersistentCoordinateFrameTrackerDestroy ::MLSDK_API::MLPersistentCoordinateFrameTrackerDestroyShim
+CREATE_FUNCTION_SHIM(ml_perception_client, const char*, MLPersistentCoordinateFrameGetResultString)
+#define MLPersistentCoordinateFrameGetResultString ::MLSDK_API::MLPersistentCoordinateFrameGetResultStringShim
+CREATE_FUNCTION_SHIM(ml_perception_client, MLResult, MLPersistentCoordinateFramesGetFrameState)
+#define MLPersistentCoordinateFramesGetFrameState ::MLSDK_API::MLPersistentCoordinateFramesGetFrameStateShim
+
+}
+
+#endif // !defined(WITH_MLSDK) || WITH_MLSDK
diff --git a/Engine/Plugins/Lumin/MLSDK/Source/MLSDK/Public/Lumin/CAPIShims/LuminAPIPlanes.h b/Engine/Plugins/Lumin/MLSDK/Source/MLSDK/Public/Lumin/CAPIShims/LuminAPIPlanes.h
new file mode 100644
--- /dev/null
+++ b/Engine/Plugins/Lumin/MLSDK/Source/MLSDK/Public/Lumin/CAPIShims/LuminAPIPlanes.h
@@ -0,0 +1,29 @@
+// Copyright 1998-2019 Epic Games, Inc. All Rights Reserved.
+
+#pragma once
+
+#if !defined(WITH_MLSDK) || WITH_MLSDK
+
+#include "Lumin/CAPIShims/LuminAPI.h"
+
+LUMIN_THIRD_PARTY_INCLUDES_START
+#include <ml_planes.h>
+LUMIN_THIRD_PARTY_INCLUDES_END
+
+namespace MLSDK_API
+{
+
+CREATE_FUNCTION_SHIM(ml_perception_client, MLResult, MLPlanesCreate)
+#define MLPlanesCreate ::MLSDK_API::MLPlanesCreateShim
+CREATE_FUNCTION_SHIM(ml_perception_client, MLResult, MLPlanesDestroy)
+#define MLPlanesDestroy ::MLSDK_API::MLPlanesDestroyShim
+CREATE_FUNCTION_SHIM(ml_perception_client, MLResult, MLPlanesQueryBegin)
+#define MLPlanesQueryBegin ::MLSDK_API::MLPlanesQueryBeginShim
+CREATE_FUNCTION_SHIM(ml_perception_client, MLResult, MLPlanesQueryGetResultsWithBoundaries)
+#define MLPlanesQueryGetResultsWithBoundaries ::MLSDK_API::MLPlanesQueryGetResultsWithBoundariesShim
+CREATE_FUNCTION_SHIM(ml_perception_client, MLResult, MLPlanesReleaseBoundariesList)
+#define MLPlanesReleaseBoundariesList ::MLSDK_API::MLPlanesReleaseBoundariesListShim
+
+}
+
+#endif // !defined(WITH_MLSDK) || WITH_MLSDK
diff --git a/Engine/Plugins/Lumin/MLSDK/Source/MLSDK/Public/Lumin/CAPIShims/LuminAPIPose.h b/Engine/Plugins/Lumin/MLSDK/Source/MLSDK/Public/Lumin/CAPIShims/LuminAPIPose.h
new file mode 100644
--- /dev/null
+++ b/Engine/Plugins/Lumin/MLSDK/Source/MLSDK/Public/Lumin/CAPIShims/LuminAPIPose.h
@@ -0,0 +1,13 @@
+// Copyright 1998-2019 Epic Games, Inc. All Rights Reserved.
+
+#pragma once
+
+#if !defined(WITH_MLSDK) || WITH_MLSDK
+
+#include "Lumin/CAPIShims/LuminAPI.h"
+
+LUMIN_THIRD_PARTY_INCLUDES_START
+#include <ml_pose.h>
+LUMIN_THIRD_PARTY_INCLUDES_END
+
+#endif // !defined(WITH_MLSDK) || WITH_MLSDK
diff --git a/Engine/Plugins/Lumin/MLSDK/Source/MLSDK/Public/Lumin/CAPIShims/LuminAPIPrivileges.h b/Engine/Plugins/Lumin/MLSDK/Source/MLSDK/Public/Lumin/CAPIShims/LuminAPIPrivileges.h
new file mode 100644
--- /dev/null
+++ b/Engine/Plugins/Lumin/MLSDK/Source/MLSDK/Public/Lumin/CAPIShims/LuminAPIPrivileges.h
@@ -0,0 +1,35 @@
+// Copyright 1998-2019 Epic Games, Inc. All Rights Reserved.
+
+#pragma once
+
+#if !defined(WITH_MLSDK) || WITH_MLSDK
+
+#include "Lumin/CAPIShims/LuminAPI.h"
+
+LUMIN_THIRD_PARTY_INCLUDES_START
+#include <ml_privilege_ids.h>
+#include <ml_privilege_functions.h>
+#include <ml_privileges.h>
+LUMIN_THIRD_PARTY_INCLUDES_END
+
+namespace MLSDK_API
+{
+
+CREATE_FUNCTION_SHIM(ml_privileges, MLResult, MLPrivilegesStartup)
+#define MLPrivilegesStartup ::MLSDK_API::MLPrivilegesStartupShim
+CREATE_FUNCTION_SHIM(ml_privileges, MLResult, MLPrivilegesShutdown)
+#define MLPrivilegesShutdown ::MLSDK_API::MLPrivilegesShutdownShim
+CREATE_FUNCTION_SHIM(ml_privileges, MLResult, MLPrivilegesCheckPrivilege)
+#define MLPrivilegesCheckPrivilege ::MLSDK_API::MLPrivilegesCheckPrivilegeShim
+CREATE_FUNCTION_SHIM(ml_privileges, MLResult, MLPrivilegesRequestPrivilege)
+#define MLPrivilegesRequestPrivilege ::MLSDK_API::MLPrivilegesRequestPrivilegeShim
+CREATE_FUNCTION_SHIM(ml_privileges, MLResult, MLPrivilegesRequestPrivilegeAsync)
+#define MLPrivilegesRequestPrivilegeAsync ::MLSDK_API::MLPrivilegesRequestPrivilegeAsyncShim
+CREATE_FUNCTION_SHIM(ml_privileges, MLResult, MLPrivilegesRequestPrivilegeTryGet)
+#define MLPrivilegesRequestPrivilegeTryGet ::MLSDK_API::MLPrivilegesRequestPrivilegeTryGetShim
+CREATE_FUNCTION_SHIM(ml_privileges, const char*, MLPrivilegesGetResultString)
+#define MLPrivilegesGetResultString ::MLSDK_API::MLPrivilegesGetResultStringShim
+
+}
+
+#endif // !defined(WITH_MLSDK) || WITH_MLSDK
diff --git a/Engine/Plugins/Lumin/MLSDK/Source/MLSDK/Public/Lumin/CAPIShims/LuminAPIPurchase.h b/Engine/Plugins/Lumin/MLSDK/Source/MLSDK/Public/Lumin/CAPIShims/LuminAPIPurchase.h
new file mode 100644
--- /dev/null
+++ b/Engine/Plugins/Lumin/MLSDK/Source/MLSDK/Public/Lumin/CAPIShims/LuminAPIPurchase.h
@@ -0,0 +1,45 @@
+// Copyright 1998-2019 Epic Games, Inc. All Rights Reserved.
+
+#pragma once
+
+#if !defined(WITH_MLSDK) || WITH_MLSDK
+
+#include "Lumin/CAPIShims/LuminAPI.h"
+
+LUMIN_THIRD_PARTY_INCLUDES_START
+#include <ml_purchase.h>
+LUMIN_THIRD_PARTY_INCLUDES_END
+
+namespace MLSDK_API
+{
+
+CREATE_FUNCTION_SHIM(ml_purchase, MLResult, MLPurchaseItemDetailsCreate)
+#define MLPurchaseItemDetailsCreate ::MLSDK_API::MLPurchaseItemDetailsCreateShim
+CREATE_FUNCTION_SHIM(ml_purchase, MLResult, MLPurchaseItemDetailsGet)
+#define MLPurchaseItemDetailsGet ::MLSDK_API::MLPurchaseItemDetailsGetShim
+CREATE_FUNCTION_SHIM(ml_purchase, MLResult, MLPurchaseItemDetailsGetResult)
+#define MLPurchaseItemDetailsGetResult ::MLSDK_API::MLPurchaseItemDetailsGetResultShim
+CREATE_FUNCTION_SHIM(ml_purchase, MLResult, MLPurchaseItemDetailsDestroy)
+#define MLPurchaseItemDetailsDestroy ::MLSDK_API::MLPurchaseItemDetailsDestroyShim
+CREATE_FUNCTION_SHIM(ml_purchase, MLResult, MLPurchaseCreate)
+#define MLPurchaseCreate ::MLSDK_API::MLPurchaseCreateShim
+CREATE_FUNCTION_SHIM(ml_purchase, MLResult, MLPurchaseSubmit)
+#define MLPurchaseSubmit ::MLSDK_API::MLPurchaseSubmitShim
+CREATE_FUNCTION_SHIM(ml_purchase, MLResult, MLPurchaseGetResult)
+#define MLPurchaseGetResult ::MLSDK_API::MLPurchaseGetResultShim
+CREATE_FUNCTION_SHIM(ml_purchase, MLResult, MLPurchaseDestroy)
+#define MLPurchaseDestroy ::MLSDK_API::MLPurchaseDestroyShim
+CREATE_FUNCTION_SHIM(ml_purchase, MLResult, MLPurchaseHistoryQueryCreate)
+#define MLPurchaseHistoryQueryCreate ::MLSDK_API::MLPurchaseHistoryQueryCreateShim
+CREATE_FUNCTION_SHIM(ml_purchase, MLResult, MLPurchaseHistoryQueryGetPage)
+#define MLPurchaseHistoryQueryGetPage ::MLSDK_API::MLPurchaseHistoryQueryGetPageShim
+CREATE_FUNCTION_SHIM(ml_purchase, MLResult, MLPurchaseHistoryQueryGetPageResult)
+#define MLPurchaseHistoryQueryGetPageResult ::MLSDK_API::MLPurchaseHistoryQueryGetPageResultShim
+CREATE_FUNCTION_SHIM(ml_purchase, MLResult, MLPurchaseHistoryQueryDestroy)
+#define MLPurchaseHistoryQueryDestroy ::MLSDK_API::MLPurchaseHistoryQueryDestroyShim
+CREATE_FUNCTION_SHIM(ml_purchase, const char*, MLPurchaseGetResultString)
+#define MLPurchaseGetResultString ::MLSDK_API::MLPurchaseGetResultStringShim
+
+}
+
+#endif // !defined(WITH_MLSDK) || WITH_MLSDK
diff --git a/Engine/Plugins/Lumin/MLSDK/Source/MLSDK/Public/Lumin/CAPIShims/LuminAPIRaycast.h b/Engine/Plugins/Lumin/MLSDK/Source/MLSDK/Public/Lumin/CAPIShims/LuminAPIRaycast.h
new file mode 100644
--- /dev/null
+++ b/Engine/Plugins/Lumin/MLSDK/Source/MLSDK/Public/Lumin/CAPIShims/LuminAPIRaycast.h
@@ -0,0 +1,27 @@
+// Copyright 1998-2019 Epic Games, Inc. All Rights Reserved.
+
+#pragma once
+
+#if !defined(WITH_MLSDK) || WITH_MLSDK
+
+#include "Lumin/CAPIShims/LuminAPI.h"
+
+LUMIN_THIRD_PARTY_INCLUDES_START
+#include <ml_raycast.h>
+LUMIN_THIRD_PARTY_INCLUDES_END
+
+namespace MLSDK_API
+{
+
+CREATE_FUNCTION_SHIM(ml_perception_client, MLResult, MLRaycastCreate)
+#define MLRaycastCreate ::MLSDK_API::MLRaycastCreateShim
+CREATE_FUNCTION_SHIM(ml_perception_client, MLResult, MLRaycastDestroy)
+#define MLRaycastDestroy ::MLSDK_API::MLRaycastDestroyShim
+CREATE_FUNCTION_SHIM(ml_perception_client, MLResult, MLRaycastRequest)
+#define MLRaycastRequest ::MLSDK_API::MLRaycastRequestShim
+CREATE_FUNCTION_SHIM(ml_perception_client, MLResult, MLRaycastGetResult)
+#define MLRaycastGetResult ::MLSDK_API::MLRaycastGetResultShim
+
+}
+
+#endif // !defined(WITH_MLSDK) || WITH_MLSDK
diff --git a/Engine/Plugins/Lumin/MLSDK/Source/MLSDK/Public/Lumin/CAPIShims/LuminAPIRemote.h b/Engine/Plugins/Lumin/MLSDK/Source/MLSDK/Public/Lumin/CAPIShims/LuminAPIRemote.h
new file mode 100644
--- /dev/null
+++ b/Engine/Plugins/Lumin/MLSDK/Source/MLSDK/Public/Lumin/CAPIShims/LuminAPIRemote.h
@@ -0,0 +1,41 @@
+// Copyright 1998-2019 Epic Games, Inc. All Rights Reserved.
+
+#pragma once
+
+#if !defined(WITH_MLSDK) || WITH_MLSDK
+
+#include "Lumin/CAPIShims/LuminAPI.h"
+#include "LuminAPIGraphics.h"
+
+LUMIN_THIRD_PARTY_INCLUDES_START
+#include <ml_remote.h>
+LUMIN_THIRD_PARTY_INCLUDES_END
+
+namespace MLSDK_API
+{
+
+CREATE_FUNCTION_SHIM(ml_remote, MLResult, MLRemoteIsServerConfigured)
+#define MLRemoteIsServerConfigured ::MLSDK_API::MLRemoteIsServerConfiguredShim
+#ifdef VK_VERSION_1_0
+CREATE_FUNCTION_SHIM(ml_remote, MLResult, MLRemoteEnumerateRequiredVkInstanceExtensions)
+#define MLRemoteEnumerateRequiredVkInstanceExtensions ::MLSDK_API::MLRemoteEnumerateRequiredVkInstanceExtensionsShim
+CREATE_FUNCTION_SHIM(ml_remote, MLResult, MLRemoteEnumerateRequiredVkDeviceExtensions)
+#define MLRemoteEnumerateRequiredVkDeviceExtensions ::MLSDK_API::MLRemoteEnumerateRequiredVkDeviceExtensionsShim
+#endif // VK_VERSION_1_0
+
+// While these functions live in the ml_graphics library, they are defined in the ml_remote header
+// and thus should be defined in LuminAPIRemote.h to avoid a circular header include dependency.
+#if PLATFORM_MAC
+CREATE_FUNCTION_SHIM(ml_graphics, MLResult, MLRemoteGraphicsCreateClientMTL)
+#define MLRemoteGraphicsCreateClientMTL ::MLSDK_API::MLRemoteGraphicsCreateClientMTLShim
+CREATE_FUNCTION_SHIM(ml_graphics, MTLPixelFormat, MLRemoteGraphicsMTLFormatFromMLSurfaceFormat)
+#define MLRemoteGraphicsMTLFormatFromMLSurfaceFormat ::MLSDK_API::MLRemoteGraphicsMTLFormatFromMLSurfaceFormatShim
+CREATE_FUNCTION_SHIM(ml_graphics, MLSurfaceFormat, MLRemoteGraphicsMLSurfaceFormatFromMTLFormat)
+#define MLRemoteGraphicsMLSurfaceFormatFromMTLFormat ::MLSDK_API::MLRemoteGraphicsMLSurfaceFormatFromMTLFormatShim
+CREATE_FUNCTION_SHIM(ml_graphics, MLResult, MLRemoteGraphicsSignalSyncObjectMTL)
+#define MLRemoteGraphicsSignalSyncObjectMTL ::MLSDK_API::MLRemoteGraphicsSignalSyncObjectMTLShim
+#endif // PLATFORM_MAC
+
+}
+
+#endif // !defined(WITH_MLSDK) || WITH_MLSDK
diff --git a/Engine/Plugins/Lumin/MLSDK/Source/MLSDK/Public/Lumin/CAPIShims/LuminAPIScreens.h b/Engine/Plugins/Lumin/MLSDK/Source/MLSDK/Public/Lumin/CAPIShims/LuminAPIScreens.h
new file mode 100644
--- /dev/null
+++ b/Engine/Plugins/Lumin/MLSDK/Source/MLSDK/Public/Lumin/CAPIShims/LuminAPIScreens.h
@@ -0,0 +1,49 @@
+// Copyright 1998-2019 Epic Games, Inc. All Rights Reserved.
+
+#pragma once
+
+#if !defined(WITH_MLSDK) || WITH_MLSDK
+
+#include "Lumin/CAPIShims/LuminAPI.h"
+
+LUMIN_THIRD_PARTY_INCLUDES_START
+#include <ml_screens.h>
+LUMIN_THIRD_PARTY_INCLUDES_END
+
+namespace MLSDK_API
+{
+
+CREATE_FUNCTION_SHIM(ml_screens, MLResult, MLScreensReleaseWatchHistoryThumbnail)
+#define MLScreensReleaseWatchHistoryThumbnail ::MLSDK_API::MLScreensReleaseWatchHistoryThumbnailShim
+CREATE_FUNCTION_SHIM(ml_screens, MLResult, MLScreensReleaseWatchHistoryList)
+#define MLScreensReleaseWatchHistoryList ::MLSDK_API::MLScreensReleaseWatchHistoryListShim
+CREATE_DEPRECATED_MSG_SHIM(ml_screens, MLResult, MLScreensReleaseScreenInfoList, "Replaced by MLScreensReleaseScreenInfoListEx.")
+#define MLScreensReleaseScreenInfoList ::MLSDK_API::MLScreensReleaseScreenInfoListShim
+CREATE_FUNCTION_SHIM(ml_screens, MLResult, MLScreensReleaseScreenInfoListEx)
+#define MLScreensReleaseScreenInfoListEx ::MLSDK_API::MLScreensReleaseScreenInfoListExShim
+CREATE_FUNCTION_SHIM(ml_screens, MLResult, MLScreensInsertWatchHistoryEntry)
+#define MLScreensInsertWatchHistoryEntry ::MLSDK_API::MLScreensInsertWatchHistoryEntryShim
+CREATE_FUNCTION_SHIM(ml_screens, MLResult, MLScreensRemoveWatchHistoryEntry)
+#define MLScreensRemoveWatchHistoryEntry ::MLSDK_API::MLScreensRemoveWatchHistoryEntryShim
+CREATE_FUNCTION_SHIM(ml_screens, MLResult, MLScreensUpdateWatchHistoryEntry)
+#define MLScreensUpdateWatchHistoryEntry ::MLSDK_API::MLScreensUpdateWatchHistoryEntryShim
+CREATE_FUNCTION_SHIM(ml_screens, MLResult, MLScreensGetWatchHistoryList)
+#define MLScreensGetWatchHistoryList ::MLSDK_API::MLScreensGetWatchHistoryListShim
+CREATE_FUNCTION_SHIM(ml_screens, MLResult, MLScreensGetWatchHistoryThumbnail)
+#define MLScreensGetWatchHistoryThumbnail ::MLSDK_API::MLScreensGetWatchHistoryThumbnailShim
+CREATE_FUNCTION_SHIM(ml_screens, MLResult, MLScreensGetScreenInfo)
+#define MLScreensGetScreenInfo ::MLSDK_API::MLScreensGetScreenInfoShim
+CREATE_FUNCTION_SHIM(ml_screens, MLResult, MLScreensUpdateScreenInfo)
+#define MLScreensUpdateScreenInfo ::MLSDK_API::MLScreensUpdateScreenInfoShim
+CREATE_DEPRECATED_MSG_SHIM(ml_screens, MLResult, MLScreensGetScreenInfoList, "Replaced by MLScreensGetScreenInfoListEx.")
+#define MLScreensGetScreenInfoList ::MLSDK_API::MLScreensGetScreenInfoListShim
+CREATE_FUNCTION_SHIM(ml_screens, MLResult, MLScreensGetScreenInfoListEx)
+#define MLScreensGetScreenInfoListEx ::MLSDK_API::MLScreensGetScreenInfoListExShim
+CREATE_FUNCTION_SHIM(ml_screens, MLResult, MLScreensCloseChannelAtScreen)
+#define MLScreensCloseChannelAtScreen ::MLSDK_API::MLScreensCloseChannelAtScreenShim
+CREATE_FUNCTION_SHIM(ml_screens, const char*, MLScreensGetResultString)
+#define MLScreensGetResultString ::MLSDK_API::MLScreensGetResultStringShim
+
+}
+
+#endif // !defined(WITH_MLSDK) || WITH_MLSDK
diff --git a/Engine/Plugins/Lumin/MLSDK/Source/MLSDK/Public/Lumin/CAPIShims/LuminAPISecureStorage.h b/Engine/Plugins/Lumin/MLSDK/Source/MLSDK/Public/Lumin/CAPIShims/LuminAPISecureStorage.h
new file mode 100644
--- /dev/null
+++ b/Engine/Plugins/Lumin/MLSDK/Source/MLSDK/Public/Lumin/CAPIShims/LuminAPISecureStorage.h
@@ -0,0 +1,29 @@
+// Copyright 1998-2019 Epic Games, Inc. All Rights Reserved.
+
+#pragma once
+
+#if !defined(WITH_MLSDK) || WITH_MLSDK
+
+#include "Lumin/CAPIShims/LuminAPI.h"
+
+LUMIN_THIRD_PARTY_INCLUDES_START
+#include <ml_secure_storage.h>
+LUMIN_THIRD_PARTY_INCLUDES_END
+
+namespace MLSDK_API
+{
+
+CREATE_FUNCTION_SHIM(ml_secure_storage, MLResult, MLSecureStoragePutBlob)
+#define MLSecureStoragePutBlob ::MLSDK_API::MLSecureStoragePutBlobShim
+CREATE_FUNCTION_SHIM(ml_secure_storage, MLResult, MLSecureStorageGetBlob)
+#define MLSecureStorageGetBlob ::MLSDK_API::MLSecureStorageGetBlobShim
+CREATE_FUNCTION_SHIM(ml_secure_storage, MLResult, MLSecureStorageDeleteBlob)
+#define MLSecureStorageDeleteBlob ::MLSDK_API::MLSecureStorageDeleteBlobShim
+CREATE_FUNCTION_SHIM(ml_secure_storage, void, MLSecureStorageFreeBlobBuffer)
+#define MLSecureStorageFreeBlobBuffer ::MLSDK_API::MLSecureStorageFreeBlobBufferShim
+CREATE_FUNCTION_SHIM(ml_secure_storage, const char*, MLSecureStorageGetResultString)
+#define MLSecureStorageGetResultString ::MLSDK_API::MLSecureStorageGetResultStringShim
+
+}
+
+#endif // !defined(WITH_MLSDK) || WITH_MLSDK
diff --git a/Engine/Plugins/Lumin/MLSDK/Source/MLSDK/Public/Lumin/CAPIShims/LuminAPISnapshot.h b/Engine/Plugins/Lumin/MLSDK/Source/MLSDK/Public/Lumin/CAPIShims/LuminAPISnapshot.h
new file mode 100644
--- /dev/null
+++ b/Engine/Plugins/Lumin/MLSDK/Source/MLSDK/Public/Lumin/CAPIShims/LuminAPISnapshot.h
@@ -0,0 +1,23 @@
+// Copyright 1998-2019 Epic Games, Inc. All Rights Reserved.
+
+#pragma once
+
+#if !defined(WITH_MLSDK) || WITH_MLSDK
+
+#include "Lumin/CAPIShims/LuminAPI.h"
+
+LUMIN_THIRD_PARTY_INCLUDES_START
+#include <ml_snapshot.h>
+LUMIN_THIRD_PARTY_INCLUDES_END
+
+namespace MLSDK_API
+{
+
+CREATE_FUNCTION_SHIM(ml_perception_client, MLResult, MLSnapshotGetTransform)
+#define MLSnapshotGetTransform ::MLSDK_API::MLSnapshotGetTransformShim
+CREATE_FUNCTION_SHIM(ml_perception_client, const char*, MLSnapshotGetResultString)
+#define MLSnapshotGetResultString ::MLSDK_API::MLSnapshotGetResultStringShim
+
+}
+
+#endif // !defined(WITH_MLSDK) || WITH_MLSDK
diff --git a/Engine/Plugins/Lumin/MLSDK/Source/MLSDK/Public/Lumin/CAPIShims/LuminAPITablet.h b/Engine/Plugins/Lumin/MLSDK/Source/MLSDK/Public/Lumin/CAPIShims/LuminAPITablet.h
new file mode 100644
--- /dev/null
+++ b/Engine/Plugins/Lumin/MLSDK/Source/MLSDK/Public/Lumin/CAPIShims/LuminAPITablet.h
@@ -0,0 +1,25 @@
+// Copyright 1998-2019 Epic Games, Inc. All Rights Reserved.
+
+#pragma once
+
+#if !defined(WITH_MLSDK) || WITH_MLSDK
+
+#include "Lumin/CAPIShims/LuminAPIInput.h"
+
+namespace MLSDK_API
+{
+
+CREATE_FUNCTION_SHIM(ml_input, MLResult, MLInputGetConnectedDevices)
+#define MLInputGetConnectedDevices ::MLSDK_API::MLInputGetConnectedDevicesShim
+CREATE_FUNCTION_SHIM(ml_input, MLResult, MLInputReleaseConnectedDevicesList)
+#define MLInputReleaseConnectedDevicesList ::MLSDK_API::MLInputReleaseConnectedDevicesListShim
+CREATE_FUNCTION_SHIM(ml_input, MLResult, MLInputSetTabletDeviceCallbacks)
+#define MLInputSetTabletDeviceCallbacks ::MLSDK_API::MLInputSetTabletDeviceCallbacksShim
+CREATE_FUNCTION_SHIM(ml_input, MLResult, MLInputGetTabletDeviceStates)
+#define MLInputGetTabletDeviceStates ::MLSDK_API::MLInputGetTabletDeviceStatesShim
+CREATE_FUNCTION_SHIM(ml_input, MLResult, MLInputReleaseTabletDeviceStates)
+#define MLInputReleaseTabletDeviceStates ::MLSDK_API::MLInputReleaseTabletDeviceStatesShim
+
+}
+
+#endif // !defined(WITH_MLSDK) || WITH_MLSDK
diff --git a/Engine/Plugins/Lumin/MLSDK/Source/MLSDK/Public/Lumin/CAPIShims/LuminAPITokenAgent.h b/Engine/Plugins/Lumin/MLSDK/Source/MLSDK/Public/Lumin/CAPIShims/LuminAPITokenAgent.h
new file mode 100644
--- /dev/null
+++ b/Engine/Plugins/Lumin/MLSDK/Source/MLSDK/Public/Lumin/CAPIShims/LuminAPITokenAgent.h
@@ -0,0 +1,29 @@
+// Copyright 1998-2019 Epic Games, Inc. All Rights Reserved.
+
+#pragma once
+
+#if !defined(WITH_MLSDK) || WITH_MLSDK
+
+#include "Lumin/CAPIShims/LuminAPI.h"
+
+LUMIN_THIRD_PARTY_INCLUDES_START
+#include <ml_token_agent.h>
+LUMIN_THIRD_PARTY_INCLUDES_END
+
+namespace MLSDK_API
+{
+
+CREATE_FUNCTION_SHIM(ml_identity, MLResult, MLTokenAgentGetClientCredentials)
+#define MLTokenAgentGetClientCredentials ::MLSDK_API::MLTokenAgentGetClientCredentialsShim
+CREATE_FUNCTION_SHIM(ml_identity, MLResult, MLTokenAgentGetClientCredentialsAsync)
+#define MLTokenAgentGetClientCredentialsAsync ::MLSDK_API::MLTokenAgentGetClientCredentialsAsyncShim
+CREATE_FUNCTION_SHIM(ml_identity, MLResult, MLTokenAgentGetClientCredentialsWait)
+#define MLTokenAgentGetClientCredentialsWait ::MLSDK_API::MLTokenAgentGetClientCredentialsWaitShim
+CREATE_FUNCTION_SHIM(ml_identity, MLResult, MLTokenAgentReleaseClientCredentials)
+#define MLTokenAgentReleaseClientCredentials ::MLSDK_API::MLTokenAgentReleaseClientCredentialsShim
+CREATE_FUNCTION_SHIM(ml_identity, const char*, MLTokenAgentGetResultString)
+#define MLTokenAgentGetResultString ::MLSDK_API::MLTokenAgentGetResultStringShim
+
+}
+
+#endif // !defined(WITH_MLSDK) || WITH_MLSDK
diff --git a/Engine/Plugins/Lumin/MLSDK/Source/MLSDK/Public/MagicLeapCFUID.h b/Engine/Plugins/Lumin/MLSDK/Source/MLSDK/Public/MagicLeapCFUID.h
new file mode 100644
--- /dev/null
+++ b/Engine/Plugins/Lumin/MLSDK/Source/MLSDK/Public/MagicLeapCFUID.h
@@ -0,0 +1,24 @@
+// Copyright 1998-2019 Epic Games, Inc. All Rights Reserved.
+
+#pragma once
+
+#include "CoreMinimal.h"
+#include "Misc/Guid.h"
+#include "Lumin/CAPIShims/LuminAPI.h"
+
+namespace MagicLeap
+{
+#if WITH_MLSDK
+    static_assert(sizeof(FGuid) >= sizeof(MLCoordinateFrameUID), "Size of FGuid should be at least as much as MLCoordinateFrameUID.");
+
+    FORCEINLINE void FGuidToMLCFUID(const FGuid& UnrealGuid, MLCoordinateFrameUID& MLGuid)
+    {
+        FMemory::Memcpy(&MLGuid, &UnrealGuid, sizeof(MLCoordinateFrameUID));
+    }
+
+    FORCEINLINE void MLCFUIDToFGuid(const MLCoordinateFrameUID& MLGuid, FGuid& UnrealGuid)
+    {
+        FMemory::Memcpy(&UnrealGuid, &MLGuid, sizeof(FGuid));
+    }
+#endif // WITH_MLSDK
+}
diff --git a/Engine/Plugins/Lumin/MLSDK/Source/MLSDK/Public/MagicLeapHandle.h b/Engine/Plugins/Lumin/MLSDK/Source/MLSDK/Public/MagicLeapHandle.h
new file mode 100644
--- /dev/null
+++ b/Engine/Plugins/Lumin/MLSDK/Source/MLSDK/Public/MagicLeapHandle.h
@@ -0,0 +1,36 @@
+// Copyright 1998-2019 Epic Games, Inc. All Rights Reserved.
+
+#pragma once
+
+#include "CoreMinimal.h"
+#include "Misc/Guid.h"
+#include "Lumin/CAPIShims/LuminAPI.h"
+
+namespace MagicLeap
+{
+#if WITH_MLSDK
+	FORCEINLINE MLHandle FGuidToMLHandle(const FGuid& UnrealGuid)
+	{
+		MLHandle MLGuid = ML_INVALID_HANDLE;
+		FMemory::Memcpy(&MLGuid, &UnrealGuid, sizeof(MLHandle));
+		return MLGuid;
+	}
+
+	FORCEINLINE FGuid MLHandleToFGuid(MLHandle MLGuid)
+	{
+		FGuid UnrealGuid;
+		FMemory::Memcpy(&UnrealGuid, &MLGuid, sizeof(MLHandle));
+		return UnrealGuid;
+	}
+#endif // WITH_MLSDK
+
+	FORCEINLINE bool FGuidIsValidHandle(const FGuid& UnrealGuid)
+	{
+#if WITH_MLSDK
+		return MLHandleIsValid(MagicLeap::FGuidToMLHandle(UnrealGuid));
+#else
+		return false;
+#endif // WITH_MLSDK
+	}
+
+}
diff --git a/Engine/Plugins/Lumin/MLSDK/Source/MLSDK/Public/MagicLeapMath.h b/Engine/Plugins/Lumin/MLSDK/Source/MLSDK/Public/MagicLeapMath.h
new file mode 100644
--- /dev/null
+++ b/Engine/Plugins/Lumin/MLSDK/Source/MLSDK/Public/MagicLeapMath.h
@@ -0,0 +1,234 @@
+// Copyright 1998-2019 Epic Games, Inc. All Rights Reserved.
+
+#pragma once
+
+#include "Kismet/KismetMathLibrary.h"
+#include "CoreMinimal.h"
+#include "Lumin/CAPIShims/LuminAPITypes.h"
+
+namespace MagicLeap
+{
+	static const float kIdentityMatColMajor[] = { 1, 0, 0, 0, 0, 1, 0, 0, 0, 0, 1, 0, 0, 0, 0, 1 };
+#if WITH_MLSDK
+	static const MLTransform kIdentityTransform = { { { { 0, 0, 0, 1 } } },{ { { 0, 0, 0 } } } };
+
+	FORCEINLINE FVector ToFVector(const MLVec3f& InVec3f, float WorldToMetersScale)
+	{
+		return FVector(-InVec3f.xyz.z * WorldToMetersScale, InVec3f.xyz.x * WorldToMetersScale, InVec3f.xyz.y * WorldToMetersScale);
+	}
+
+	FORCEINLINE FVector ToFVectorNoScale(const MLVec3f& InVec3f)
+	{
+		return FVector(-InVec3f.xyz.z, InVec3f.xyz.x, InVec3f.xyz.y);
+	}
+
+	FORCEINLINE FVector ToFVectorExtents(const MLVec3f& InVec3f, float WorldToMetersScale)
+	{
+		// No sign change for FVector::X
+		return FVector(InVec3f.xyz.z * WorldToMetersScale, InVec3f.xyz.x * WorldToMetersScale, InVec3f.xyz.y * WorldToMetersScale);
+	}
+
+	FORCEINLINE MLVec3f ToMLVector(const FVector& InVector, float WorldToMetersScale)
+	{
+		float InverseScale = 1.0f / WorldToMetersScale;
+		MLVec3f vec;
+		vec.xyz.x = InVector.Y * InverseScale;
+		vec.xyz.y = InVector.Z * InverseScale;
+		vec.xyz.z = -InVector.X * InverseScale;
+		return vec;
+	}
+
+	FORCEINLINE MLVec3f ToMLVectorNoScale(const FVector& InVector)
+	{
+		return ToMLVector(InVector, 1.0f);
+	}
+
+	FORCEINLINE MLVec3f ToMLVectorExtents(const FVector& InVector, float WorldToMetersScale)
+	{
+		float InverseScale = 1.0f / WorldToMetersScale;
+		MLVec3f vec;
+		vec.xyz.x = InVector.Y * InverseScale;
+		vec.xyz.y = InVector.Z * InverseScale;
+		// No sign change for MLVec3f::z
+		vec.xyz.z = InVector.X * InverseScale;
+		return vec;
+	}
+
+	FORCEINLINE FQuat ToFQuat(const MLQuaternionf& InQuat)
+	{
+		return FQuat(-InQuat.z, InQuat.x, InQuat.y, -InQuat.w);
+	}
+
+	FORCEINLINE MLQuaternionf ToMLQuat(const FQuat& InQuat)
+	{
+		MLQuaternionf quat;
+		quat.x = InQuat.Y;
+		quat.y = InQuat.Z;
+		quat.z = -InQuat.X;
+		quat.w = -InQuat.W;
+		return quat;
+	}
+
+	FORCEINLINE FTransform ToFTransform(const MLTransform& InTransform, float WorldToMetersScale)
+	{
+		return FTransform(ToFQuat(InTransform.rotation), ToFVector(InTransform.position, WorldToMetersScale), FVector(1.0f, 1.0f, 1.0f));
+	}
+
+	FORCEINLINE MLTransform ToMLTransform(const FTransform& InTransform, float WorldToMetersScale)
+	{
+		MLTransform transform = { { { { 0, 0, 0, 1 } } }, { { { 0, 0, 0 } } } };
+		transform.position = ToMLVector(InTransform.GetLocation(), WorldToMetersScale);
+		transform.rotation = ToMLQuat(InTransform.GetRotation());
+		return transform;
+	}
+
+	// Simple col major to row major conversion
+	FORCEINLINE FMatrix ToFMatrix(const MLMat4f& InMat4f)
+	{
+		return FMatrix(
+			FPlane(InMat4f.matrix_colmajor[0], InMat4f.matrix_colmajor[1], InMat4f.matrix_colmajor[2], InMat4f.matrix_colmajor[3]),
+			FPlane(InMat4f.matrix_colmajor[4], InMat4f.matrix_colmajor[5], InMat4f.matrix_colmajor[6], InMat4f.matrix_colmajor[7]),
+			FPlane(InMat4f.matrix_colmajor[8], InMat4f.matrix_colmajor[9], InMat4f.matrix_colmajor[10], InMat4f.matrix_colmajor[11]),
+			FPlane(InMat4f.matrix_colmajor[12], InMat4f.matrix_colmajor[13], InMat4f.matrix_colmajor[14], InMat4f.matrix_colmajor[15])
+		);
+	}
+
+	FORCEINLINE FMatrix ToFMatrix2(const MLMat4f& InMat4f, const float WorldToMetersScale)
+	{
+		return FMatrix(
+			FPlane(-InMat4f.matrix_colmajor[10], InMat4f.matrix_colmajor[8], InMat4f.matrix_colmajor[9], 0.0f),
+			FPlane(-InMat4f.matrix_colmajor[2], InMat4f.matrix_colmajor[0], InMat4f.matrix_colmajor[1], 0.0f),
+			FPlane(-InMat4f.matrix_colmajor[6], InMat4f.matrix_colmajor[4], InMat4f.matrix_colmajor[5], 0.0f),
+			FPlane(-InMat4f.matrix_colmajor[14] * WorldToMetersScale, InMat4f.matrix_colmajor[12] * WorldToMetersScale, InMat4f.matrix_colmajor[13] * WorldToMetersScale, 1.0f)
+		);
+	}
+
+	// Convert from MLGraphics Projection Model to Unreal's Reverse Infinite Z Projection Model 	
+	FORCEINLINE FMatrix ToUEProjectionMatrix(const MLMat4f& InProjMat, float NearClipPlane)
+	{
+		/*
+			* matrix_colmajor indices *
+
+			| 00  04  08  12 |
+			| 01  05  09  13 |
+			| 02  06  10  14 |
+			| 03  07  11  15 |
+
+			* ml_projection_colmajor format *
+			* A, B depends on requested Z format *
+
+			| X  0  0   0 |
+			| 0  Y  0   0 |
+			| 0  0  A   B |
+			| 0  0  -1  0 |
+
+			* ue_projection_rowmajor format *
+			* A = 0, B = N (NearClip) *
+
+			* Unreal is left handed and MLGraphics is right handed so ue_projection_rowmajor[2] = -1.0 x ue_projection_rowmajor[2] *
+			* which makes ue_projection_rowmajor[2][3] = 1.0f instead of -1.0f. *
+			* Unreal expects the projection matrix to be in centimeters and uses it for various purposes *
+			* such as bounding volume calculations for lights in the shadow algorithm. *
+			* We're overwriting the near value to match the units of unreal here instead of using the units of the MLSDK. *
+
+			| X  0  0  0 |
+			| 0  Y  0  0 |
+			| 0  0  0  1 |
+			| 0  0  N  0 |
+		*/
+
+		return FMatrix(
+			FPlane(InProjMat.matrix_colmajor[0], InProjMat.matrix_colmajor[1], InProjMat.matrix_colmajor[2], InProjMat.matrix_colmajor[3]),
+			FPlane(InProjMat.matrix_colmajor[4], InProjMat.matrix_colmajor[5], InProjMat.matrix_colmajor[6], InProjMat.matrix_colmajor[7]),
+			FPlane(-InProjMat.matrix_colmajor[8], -InProjMat.matrix_colmajor[9],			0.0f		   ,			1.0f			 ),
+			FPlane(InProjMat.matrix_colmajor[12], InProjMat.matrix_colmajor[13],		NearClipPlane	   , InProjMat.matrix_colmajor[15])
+		);
+	}
+
+	FORCEINLINE FVector ToFVector(const MLMat4f& InMatrix, float WorldToMetersScale)
+	{
+		MLVec3f MLPosition;
+		MLPosition.xyz.x = InMatrix.matrix_colmajor[12];
+		MLPosition.xyz.y = InMatrix.matrix_colmajor[13];
+		MLPosition.xyz.z = InMatrix.matrix_colmajor[14];
+		return ToFVector(MLPosition, WorldToMetersScale);
+	}
+
+	FORCEINLINE FQuat ToFQuat(const MLMat4f& InMatrix)
+	{
+		const float Trace = InMatrix.matrix_colmajor[0] + InMatrix.matrix_colmajor[5] + InMatrix.matrix_colmajor[10];
+
+		MLQuaternionf Output;
+		if (Trace > 0)
+		{
+			const float S = FMath::Sqrt(Trace + 1.0) * 2;
+			Output.w = S * 0.25;
+			Output.x = (InMatrix.matrix_colmajor[6] - InMatrix.matrix_colmajor[9]) / S;
+			Output.y = (InMatrix.matrix_colmajor[8] - InMatrix.matrix_colmajor[2]) / S;
+			Output.z = (InMatrix.matrix_colmajor[1] - InMatrix.matrix_colmajor[4]) / S;
+		}
+		else if (InMatrix.matrix_colmajor[0] > InMatrix.matrix_colmajor[5] && InMatrix.matrix_colmajor[0] > InMatrix.matrix_colmajor[10])
+		{
+			const float S = FMath::Sqrt(InMatrix.matrix_colmajor[0] - InMatrix.matrix_colmajor[5] - InMatrix.matrix_colmajor[10] + 1.0) * 2;
+			Output.x = S * 0.25;
+			Output.w = (InMatrix.matrix_colmajor[6] - InMatrix.matrix_colmajor[9]) / S;
+			Output.y = (InMatrix.matrix_colmajor[1] - InMatrix.matrix_colmajor[4]) / S;
+			Output.z = (InMatrix.matrix_colmajor[8] - InMatrix.matrix_colmajor[2]) / S;
+		}
+		else if (InMatrix.matrix_colmajor[5] > InMatrix.matrix_colmajor[10])
+		{
+			const float S = FMath::Sqrt(InMatrix.matrix_colmajor[5] - InMatrix.matrix_colmajor[0] - InMatrix.matrix_colmajor[10] + 1.0) * 2;
+			Output.y = S * 0.25;
+			Output.w = (InMatrix.matrix_colmajor[8] - InMatrix.matrix_colmajor[2]) / S;
+			Output.x = (InMatrix.matrix_colmajor[1] - InMatrix.matrix_colmajor[4]) / S;
+			Output.z = (InMatrix.matrix_colmajor[6] - InMatrix.matrix_colmajor[9]) / S;
+		}
+		else
+		{
+			const float S = FMath::Sqrt(InMatrix.matrix_colmajor[10] - InMatrix.matrix_colmajor[5] - InMatrix.matrix_colmajor[0] + 1.0) * 2;
+			Output.z = S * 0.25;
+			Output.w = (InMatrix.matrix_colmajor[1] - InMatrix.matrix_colmajor[4]) / S;
+			Output.x = (InMatrix.matrix_colmajor[8] - InMatrix.matrix_colmajor[2]) / S;
+			Output.y = (InMatrix.matrix_colmajor[6] - InMatrix.matrix_colmajor[9]) / S;
+		}
+
+		return ToFQuat(Output);
+	}
+
+	FORCEINLINE MLVec3f ExtractScale3D(const MLMat4f& InMatrix)
+	{
+		MLVec3f MLScale;
+		MLScale.x = FMath::Sqrt(
+			FMath::Square(InMatrix.matrix_colmajor[0])
+			+ FMath::Square(InMatrix.matrix_colmajor[1])
+			+ FMath::Square(InMatrix.matrix_colmajor[2]));
+		MLScale.y = FMath::Sqrt(
+			FMath::Square(InMatrix.matrix_colmajor[4])
+			+ FMath::Square(InMatrix.matrix_colmajor[5])
+			+ FMath::Square(InMatrix.matrix_colmajor[6]));
+		MLScale.z = FMath::Sqrt(
+			FMath::Square(InMatrix.matrix_colmajor[8])
+			+ FMath::Square(InMatrix.matrix_colmajor[9])
+			+ FMath::Square(InMatrix.matrix_colmajor[10]));
+		return MLScale;
+	}
+
+	FORCEINLINE MLMat4f ToMLMat4f(const FVector& InPosition, const FQuat& InOrientation, const FVector& InScale3D, const float WorldToMetersScale)
+	{
+		FMatrix OutMatrix = FTransform(InOrientation, InPosition, InScale3D).ToMatrixWithScale();
+		return MLMat4f {
+			{
+				OutMatrix.M[1][1], OutMatrix.M[1][2], -OutMatrix.M[1][0], 0.0f,
+				OutMatrix.M[2][1], OutMatrix.M[2][2], -OutMatrix.M[2][0], 0.0f,
+				OutMatrix.M[0][1], OutMatrix.M[0][2], -OutMatrix.M[0][0], 0.0f,
+				OutMatrix.M[3][1] / WorldToMetersScale, OutMatrix.M[3][2] / WorldToMetersScale, -OutMatrix.M[3][0] / WorldToMetersScale, 1.0f
+			}
+		};
+	}
+#endif //WITH_MLSDK
+
+	FORCEINLINE FQuat ToUERotator(const FQuat& InMLRotation)
+	{
+		return FQuat(UKismetMathLibrary::GetRightVector(InMLRotation.Rotator()), PI) * InMLRotation;
+	}
+}
diff --git a/Engine/Plugins/Lumin/MLSDK/Source/MLSDK/Public/MagicLeapPluginUtil.h b/Engine/Plugins/Lumin/MLSDK/Source/MLSDK/Public/MagicLeapPluginUtil.h
new file mode 100644
--- /dev/null
+++ b/Engine/Plugins/Lumin/MLSDK/Source/MLSDK/Public/MagicLeapPluginUtil.h
@@ -0,0 +1,126 @@
+// Copyright 1998-2019 Epic Games, Inc. All Rights Reserved.
+
+#pragma once
+#include "CoreMinimal.h"
+#include "Misc/ConfigCacheIni.h"
+#include "HAL/PlatformProcess.h"
+#include "Misc/Paths.h"
+#include "GenericPlatform/GenericPlatformFile.h"
+#include "Misc/CString.h"
+#include "Templates/Atomic.h"
+#include "Lumin/CAPIShims/LuminAPIPlatform.h"
+
+/** Utility class to deal with some API features. */
+class FMagicLeapAPISetup
+{
+public:
+	/** Fills the Variables with the evaluated contents of the SDK Shim discovery data. */
+	bool GetZIShimVariables(const FString & MLSDK, TMap<FString, FString> & ResultVariables) const
+	{
+		// The known path to the paths file.
+		FString SDKShimDiscoveryFile = FPaths::Combine(MLSDK, TEXT(".metadata"), TEXT("sdk_shim_discovery.txt"));
+		if (!FPaths::FileExists(SDKShimDiscoveryFile))
+		{
+			return false;
+		}
+		// Map of variable to value for evaluating the content of the file.
+		// On successful parsing and evaluation we copy the data to the result.
+		TMap<FString, FString> Variables;
+		Variables.Add(TEXT("$(MLSDK)"), MLSDK);
+#if PLATFORM_WINDOWS
+		Variables.Add(TEXT("$(HOST)"), TEXT("win64"));
+#elif PLATFORM_LINUX
+		Variables.Add(TEXT("$(HOST)"), TEXT("linux64"));
+#elif PLATFORM_MAC
+		Variables.Add(TEXT("$(HOST)"), TEXT("osx"));
+#endif // PLATFORM_WINDOWS
+		// Single pass algo for evaluating the file:
+		// 1. for each line:
+		// a. for each occurance of $(NAME):
+		// i. replace with Variables[NAME], if no NAME var found ignore replacement.
+		// b. split into var=value, and add to Variables.
+		IPlatformFile& PlatformFile = IPlatformFile::GetPlatformPhysical();
+		TUniquePtr<IFileHandle> File(PlatformFile.OpenRead(*SDKShimDiscoveryFile));
+		if (File)
+		{
+			TArray<uint8> Data;
+			Data.SetNumZeroed(File->Size() + 1);
+			File->Read(Data.GetData(), Data.Num() - 1);
+			FUTF8ToTCHAR TextConvert(reinterpret_cast<ANSICHAR*>(Data.GetData()));
+			const TCHAR * Text = TextConvert.Get();
+			while (Text && *Text)
+			{
+				Text += FCString::Strspn(Text, TEXT("\t "));
+				if (Text[0] == '#' || Text[0] == '\r' || Text[0] == '\n' || !*Text)
+				{
+					// Skip comment or empty lines.
+					Text += FCString::Strcspn(Text, TEXT("\r\n"));
+				}
+				else
+				{
+					// Parse variable=value line.
+					FString Variable(FCString::Strcspn(Text, TEXT("\t= ")), Text);
+					Text += Variable.Len();
+					Text += FCString::Strspn(Text, TEXT("\t= "));
+					FString Value(FCString::Strcspn(Text, TEXT("\r\n")), Text);
+					Text += Value.Len();
+					Value.TrimEndInline();
+					// Eval any var references in both variable and value.
+					int32 EvalCount = 0;
+					do
+					{
+						EvalCount = 0;
+						for (auto VarEntry : Variables)
+						{
+							EvalCount += Variable.ReplaceInline(*VarEntry.Key, *VarEntry.Value);
+							EvalCount += Value.ReplaceInline(*VarEntry.Key, *VarEntry.Value);
+						}
+					} while (EvalCount != 0 && (Variable.Contains(TEXT("$(")) || Value.Contains(TEXT("$("))));
+					// Intern the new variable.
+					Variable = TEXT("$(") + Variable + TEXT(")");
+					Variables.Add(Variable, Value);
+				}
+				// Skip over EOL to get to the next line.
+				if (Text[0] == '\r' && Text[1] == '\n')
+				{
+					Text += 2;
+				}
+				else
+				{
+					Text += 1;
+				}
+			}
+		}
+		// We now need to copy the evaled variables to the result and un-munge them for
+		// plain access. We use them munged for the eval for easier eval replacement above.
+		ResultVariables.Empty(Variables.Num());
+		for (auto VarEntry : Variables)
+		{
+			ResultVariables.Add(VarEntry.Key.Mid(2, VarEntry.Key.Len() - 3), VarEntry.Value);
+		}
+		return true;
+	}
+
+	/** Returns the, cached, API level for the platform. (thread safe) */
+	static uint32 GetPlatformLevel()
+	{
+#if WITH_MLSDK
+		// Default to the compiled platform level if we can't get the runtime value.
+		static TAtomic<uint32> Level(ML_PLATFORM_API_LEVEL);
+		static TAtomic<bool> HaveLevel(false);
+		if (!HaveLevel.Load(EMemoryOrder::Relaxed))
+		{
+			uint32 TempLevel = 0;
+			MLResult Result = MLPlatformGetAPILevel(&TempLevel);
+			if (Result == MLResult_Ok)
+			{
+				Level = TempLevel;
+			}
+			HaveLevel = true;
+		}
+		return Level;
+#else
+		return 0;
+#endif
+	}
+};
diff --git a/Engine/Plugins/Lumin/MLSDK/Source/MLSDK/Public/MagicLeapUtil.h b/Engine/Plugins/Lumin/MLSDK/Source/MLSDK/Public/MagicLeapUtil.h
new file mode 100644
--- /dev/null
+++ b/Engine/Plugins/Lumin/MLSDK/Source/MLSDK/Public/MagicLeapUtil.h
@@ -0,0 +1,23 @@
+// Copyright 1998-2019 Epic Games, Inc. All Rights Reserved.
+
+#pragma once
+
+#include "Logging/LogMacros.h"
+
+#define ML_LOG_ONCE(CategoryName, Verbosity, Format, ...) \
+	do { static bool _logged = false; if (!_logged) { _logged = true; UE_LOG(CategoryName, Verbosity, Format, ##__VA_ARGS__); } } while (0);
+
+namespace MagicLeap
+{
+	template<typename T> inline const FString EnumToString(const char* EnumName, T EnumValue)
+	{
+		static const FString InvalidEnum("Invalid");
+
+		const UEnum* EnumPtr = FindObject<UEnum>(ANY_PACKAGE, UTF8_TO_TCHAR(EnumName), true);
+		if (!EnumPtr)
+		{
+			return InvalidEnum;
+		}
+		return EnumPtr->GetNameStringByIndex(static_cast<int32>(EnumValue));
+	}
+}
diff --git a/Engine/Plugins/Lumin/MagicLeap/Config/BaseMagicLeap.ini b/Engine/Plugins/Lumin/MagicLeap/Config/BaseMagicLeap.ini
--- a/Engine/Plugins/Lumin/MagicLeap/Config/BaseMagicLeap.ini
+++ b/Engine/Plugins/Lumin/MagicLeap/Config/BaseMagicLeap.ini
@@ -1,3 +1,86 @@
 [CoreRedirects]
+; General MagicLeap redirects
++EnumRedirects=(OldName="EFailReason", NewName="/Script/MagicLeap.EMagicLeapTransformFailReason")
 
-+StructRedirects=(OldName="/Script/MagicLeap.MeshInfo", NewName="/Script/MagicLeap.MLTrackingMeshInfo")
+; Head tracking redirects
++EnumRedirects=(OldName="EHeadTrackingError", NewName="/Script/MagicLeap.EMagicLeapHeadTrackingError")
++EnumRedirects=(OldName="EHeadTrackingMode", NewName="/Script/MagicLeap.EMagicLeapHeadTrackingMode")
++EnumRedirects=(OldName="EHeadTrackingMapEvent", NewName="/Script/MagicLeap.EMagicLeapHeadTrackingMapEvent")
++StructRedirects=(OldName="HeadTrackingState", NewName="/Script/MagicLeap.MagicLeapHeadTrackingState")
++ClassRedirects=(OldName="/Script/MagicLeap.HeadTrackingNotificationsComponent", NewName="/Script/MagicLeap.MagicLeapHeadTrackingNotificationsComponent")
+
+; Meshing redirects
++EnumRedirects=(OldName="EMeshType", NewName="/Script/MagicLeap.EMagicLeapMeshType")
++EnumRedirects=(OldName="EMLMeshVertexColorMode", NewName="/Script/MagicLeap.EMagicLeapMeshVertexColorMode")
++EnumRedirects=(OldName="EMeshLOD", NewName="/Script/MagicLeap.EMagicLeapMeshLOD")
++EnumRedirects=(OldName="EMeshState", NewName="/Script/MagicLeap.EMagicLeapMeshState")
++StructRedirects=(OldName="MeshBlockInfo", NewName="/Script/MagicLeap.MagicLeapMeshBlockInfo")
++StructRedirects=(OldName="MeshInfo", NewName="/Script/MagicLeap.MagicLeapTrackingMeshInfo")
++StructRedirects=(OldName="MLTrackingMeshInfo", NewName="/Script/MagicLeap.MagicLeapTrackingMeshInfo")
++StructRedirects=(OldName="MeshBlockRequest", NewName="/Script/MagicLeap.MagicLeapMeshBlockRequest")
++ClassRedirects=(OldName="/Script/MagicLeap.MeshBlockSelectorInterface", NewName="/Script/MagicLeap.MagicLeapMeshBlockSelectorInterface")
+
+; Raycast redirects
++StructRedirects=(OldName="RaycastQueryParams", NewName="/Script/MagicLeap.MagicLeapRaycastQueryParams")
++EnumRedirects=(OldName="ERaycastResultState", NewName="/Script/MagicLeap.EMagicLeapRaycastResultState")
++StructRedirects=(OldName="RaycastHitResult", NewName="/Script/MagicLeap.MagicLeapRaycastHitResult")
++ClassRedirects=(OldName="/Script/MagicLeap.RaycastComponent", NewName="/Script/MagicLeap.MagicLeapRaycastComponent")
++ClassRedirects=(OldName="/Script/MagicLeap.RaycastFunctionLibrary", NewName="/Script/MagicLeap.MagicLeapRaycastFunctionLibrary")
+
+; ARPin redirects
++EnumRedirects=(OldName="EPassableWorldError", NewName="/Script/MagicLeapARPin.EMagicLeapPassableWorldError")
++EnumRedirects=(OldName="EAutoPinType", NewName="/Script/MagicLeapARPin.EMagicLeapAutoPinType")
++ClassRedirects=(OldName="/Script/MagicLeap.MagicLeapARPinSaveGame", NewName="/Script/MagicLeapARPin.MagicLeapARPinSaveGame")
++ClassRedirects=(OldName="/Script/MagicLeap.MagicLeapARPinFunctionLibrary", NewName="/Script/MagicLeapARPin.MagicLeapARPinFunctionLibrary")
++ClassRedirects=(OldName="/Script/MagicLeap.MagicLeapARPinComponent", NewName="/Script/MagicLeapARPin.MagicLeapARPinComponent")
+
+; Audio redirects
++ClassRedirects=(OldName="/Script/MagicLeapAudio.AudioJackPluggedDelegate", NewName="/Script/MagicLeapAudio.MagicLeapAudioJackPluggedDelegate")
++ClassRedirects=(OldName="/Script/MagicLeapAudio.AudioJackPluggedDelegateMulti", NewName="/Script/MagicLeapAudio.MagicLeapAudioJackPluggedDelegateMulti")
++ClassRedirects=(OldName="/Script/MagicLeapAudio.AudioJackUnpluggedDelegate", NewName="/Script/MagicLeapAudio.MagicLeapAudioJackUnpluggedDelegate")
++ClassRedirects=(OldName="/Script/MagicLeapAudio.AudioJackUnpluggedDelegateMulti", NewName="/Script/MagicLeapAudio.MagicLeapAudioJackUnpluggedDelegateMulti")
+
+; Controller redirects
++EnumRedirects=(OldName="EMLControllerType", NewName="/Script/MagicLeapController.EMagicLeapControllerType")
++EnumRedirects=(OldName="EMLControllerLEDPattern", NewName="/Script/MagicLeapController.EMagicLeapControllerLEDPattern")
++EnumRedirects=(OldName="EMLControllerLEDEffect", NewName="/Script/MagicLeapController.EMagicLeapControllerLEDEffect")
++EnumRedirects=(OldName="EMLControllerLEDColor", NewName="/Script/MagicLeapController.EMagicLeapControllerLEDColor")
++EnumRedirects=(OldName="EMLControllerLEDSpeed", NewName="/Script/MagicLeapController.EMagicLeapControllerLEDSpeed")
++EnumRedirects=(OldName="EMLControllerHapticPattern", NewName="/Script/MagicLeapController.EMagicLeapControllerHapticPattern")
++EnumRedirects=(OldName="EMLControllerHapticIntensity", NewName="/Script/MagicLeapController.EMagicLeapControllerHapticIntensity")
++EnumRedirects=(OldName="EMLControllerTrackingMode", NewName="/Script/MagicLeapController.EMagicLeapControllerTrackingMode")
++ClassRedirects=(OldName="/Script/MagicLeapController.TouchpadGesturesComponent", NewName="/Script/MagicLeapController.MagicLeapTouchpadGesturesComponent")
+
+; Hand meshing redirects
++StructRedirects=(OldName="HandMeshBlock", NewName="/Script/MagicLeapHandMeshing.MagicLeapHandMeshBlock")
++StructRedirects=(OldName="HandMesh", NewName="/Script/MagicLeapHandMeshing.MagicLeapHandMesh")
+
+; Hand tracking redirects
++EnumRedirects=(OldName="EHandTrackingKeypoint", NewName="/Script/MagicLeapHandTracking.EMagicLeapHandTrackingKeypoint")
++EnumRedirects=(OldName="EHandTrackingGesture", NewName="/Script/MagicLeapHandTracking.EMagicLeapHandTrackingGesture")
++EnumRedirects=(OldName="EHandTrackingKeypointFilterLevel", NewName="/Script/MagicLeapHandTracking.EMagicLeapHandTrackingKeypointFilterLevel")
++EnumRedirects=(OldName="EHandTrackingGestureFilterLevel", NewName="/Script/MagicLeapHandTracking.EMagicLeapHandTrackingGestureFilterLevel")
++EnumRedirects=(OldName="EGestureTransformSpace", NewName="/Script/MagicLeapHandTracking.EMagicLeapGestureTransformSpace")
+
+; Planes redirects
++EnumRedirects=(OldName="EPlaneQueryFlags", NewName="/Script/MagicLeapPlanes.EMagicLeapPlaneQueryFlags")
++StructRedirects=(OldName="PlaneResult", NewName="/Script/MagicLeapPlanes.MagicLeapPlaneResult")
++StructRedirects=(OldName="PlanesQuery", NewName="/Script/MagicLeapPlanes.MagicLeapPlanesQuery")
++StructRedirects=(OldName="MLPolygon", NewName="/Script/MagicLeapPlanes.MagicLeapPolygon")
++StructRedirects=(OldName="PlaneBoundary", NewName="/Script/MagicLeapPlanes.MagicLeapPlaneBoundary")
++StructRedirects=(OldName="PlaneBoundaries", NewName="/Script/MagicLeapPlanes.MagicLeapPlaneBoundaries")
++ClassRedirects=(OldName="/Script/MagicLeapPlanes.PlanesResultStaticDelegate", NewName="/Script/MagicLeapPlanes.MagicLeapPlanesResultStaticDelegate")
++ClassRedirects=(OldName="/Script/MagicLeapPlanes.PlanesResultDelegate", NewName="/Script/MagicLeapPlanes.MagicLeapPlanesResultDelegate")
++ClassRedirects=(OldName="/Script/MagicLeapPlanes.PlanesResultDelegateMulti", NewName="/Script/MagicLeapPlanes.MagicLeapPlanesResultDelegateMulti")
+
+; LightingTracking redirects
++ClassRedirects=(OldName="/Script/MagicLeap.LightingTrackingComponent", NewName="/Script/MagicLeapLightEstimation.MagicLeapLightingTrackingComponent")
+
+; Privileges redirects
++EnumRedirects=(OldName="/Script/MagicLeap.EMagicLeapPrivilege", NewName="/Script/MagicLeapPrivileges.EMagicLeapPrivilege")
++ClassRedirects=(OldName="/Script/MagicLeapPrivileges.PrivilegeRequestStaticDelegate", NewName="/Script/MagicLeapLightEstimation.MagicLeapPrivilegeRequestStaticDelegate")
++ClassRedirects=(OldName="/Script/MagicLeapPrivileges.PrivilegeRequestDelegate", NewName="/Script/MagicLeapLightEstimation.MagicLeapPrivilegeRequestDelegate")
+
+; ImageTracking redirects
++ClassRedirects=(OldName="/Script/MagicLeap.ImageTrackerComponent", NewName="/Script/MagicLeapImageTracker.MagicLeapImageTrackerComponent")
++ClassRedirects=(OldName="/Script/MagicLeap.ImageTrackerFunctionLibrary", NewName="/Script/MagicLeapImageTracker.MagicLeapImageTrackerFunctionLibrary")
diff --git a/Engine/Plugins/Lumin/MagicLeap/MagicLeap.uplugin b/Engine/Plugins/Lumin/MagicLeap/MagicLeap.uplugin
--- a/Engine/Plugins/Lumin/MagicLeap/MagicLeap.uplugin
+++ b/Engine/Plugins/Lumin/MagicLeap/MagicLeap.uplugin
@@ -15,7 +15,9 @@
 	"IsBetaVersion" : false,
 	"Installed": false,
 	"SupportedTargetPlatforms": [
-		"Lumin"
+		"Lumin",
+		"Mac",
+		"Win64"
 	],
 	"Modules" :
 	[
@@ -34,8 +36,17 @@
 			"Name": "MagicLeapAR",
 			"Type": "Runtime",
 			"LoadingPhase": "PreDefault",
-			"WhitelistPlatforms": 
-			[
+			"WhitelistPlatforms": [
+				"Win64",
+				"Mac",
+				"Lumin"
+			]
+		},
+		{
+			"Name": "MagicLeapARPin",
+			"Type": "Runtime",
+			"LoadingPhase": "PreDefault",
+			"WhitelistPlatforms": [
 				"Win64",
 				"Mac",
 				"Lumin"
@@ -46,6 +57,7 @@
 			"Type" : "Runtime",
 			"WhitelistPlatforms": [
 				"Win64",
+				"Mac",
 				"Lumin"
 			]
 		},
@@ -78,6 +90,17 @@
 				"Lumin"
 			]
 		},
+		{
+			"Name" : "MagicLeapHandMeshing",
+			"Type" : "Runtime",
+			"LoadingPhase" : "PreDefault",
+			"WhitelistPlatforms" :
+			[
+				"Win64",
+				"Mac",
+				"Lumin"
+			]
+		},
 		{
 			"Name" : "MagicLeapHandTracking",
 			"Type" : "Runtime",
@@ -89,6 +112,17 @@
 				"Lumin"
 			]
 		},
+		{
+			"Name" : "MagicLeapImageTracker",
+			"Type" : "Runtime",
+			"LoadingPhase" : "PreDefault",
+			"WhitelistPlatforms" :
+			[
+				"Win64",
+				"Mac",
+				"Lumin"
+			]
+		},
 		{
 			"Name": "MagicLeapSecureStorage",
 			"Type": "Runtime",
@@ -131,11 +165,31 @@
 				"Lumin"
 			]
 		},
+		{
+			"Name": "MagicLeapPlanes",
+			"Type": "Runtime",
+			"LoadingPhase": "PreDefault",
+			"WhitelistPlatforms": [
+				"Win64",
+				"Mac",
+				"Lumin"
+			]
+		},
+		{
+			"Name": "MagicLeapPrivileges",
+			"Type": "Runtime",
+			"LoadingPhase": "PreDefault",
+			"WhitelistPlatforms": [
+				"Win64",
+				"Mac",
+				"Lumin"
+			]
+		}
 	],
 	"Plugins":
 	[
 		{
-			"Name": "ProceduralMeshComponent",
+			"Name": "MLSDK",
 			"Enabled": true
 		}
 	]
diff --git a/Engine/Plugins/Lumin/MagicLeap/Source/MagicLeap/MagicLeap.Build.cs b/Engine/Plugins/Lumin/MagicLeap/Source/MagicLeap/MagicLeap.Build.cs
--- a/Engine/Plugins/Lumin/MagicLeap/Source/MagicLeap/MagicLeap.Build.cs
+++ b/Engine/Plugins/Lumin/MagicLeap/Source/MagicLeap/MagicLeap.Build.cs
@@ -24,10 +24,11 @@ namespace UnrealBuildTool.Rules
 				new string[]
 				{
 					"HeadMountedDisplay",
-					"ProceduralMeshComponent",
 					"InputDevice",
                     "LuminRuntimeSettings",
 					"AugmentedReality",
+					"MLSDK",
+					"MagicLeapPrivileges"
                 }
 			);
 
@@ -44,12 +45,10 @@ namespace UnrealBuildTool.Rules
 					"Renderer",
 					"Slate",
 					"SlateCore",
-					"MLSDK",
 					"MRMesh",
                     "MagicLeapHelperOpenGL",
 					// Public headers of MagicLeapHelperVulkan are protected against Mac so this is fine here.
 					"MagicLeapHelperVulkan",
-					"LuminRuntimeSettings",
 					"MagicLeapSecureStorage",
 				}
 			);
@@ -62,6 +61,16 @@ namespace UnrealBuildTool.Rules
 						"VulkanRHI"
 					}
 				);
+                PrivateDefinitions.Add("MLSDK_API_USE_VULKAN=1");
+			}
+			else
+			{
+				PrivateDependencyModuleNames.Add("MetalRHI");
+				PrivateIncludePaths.AddRange(
+					new string[] {
+						Path.Combine(EngineSourceDirectory, "Runtime/Apple/MetalRHI/Private"),
+					});
+				AddEngineThirdPartyPrivateStaticDependencies(Target, "MTLPP");
 			}
 
 			if (Target.bBuildEditor == true)
@@ -95,6 +104,14 @@ namespace UnrealBuildTool.Rules
 						Path.Combine(EngineSourceDirectory, "ThirdParty/SDL2/SDL-gui-backend/include"),
 					}
 				);
+				PrivateIncludePaths.AddRange(
+					new string[] {
+						Path.Combine(EngineSourceDirectory, "Runtime/VulkanRHI/Private"),
+						Path.Combine(EngineSourceDirectory, "Runtime/VulkanRHI/Private/Linux")
+					}
+				);
+				AddEngineThirdPartyPrivateStaticDependencies(Target, "Vulkan");
+
 				// HACK This is a workaround for a bug introduced in Unreal 4.13 which should be
 				// removed if all SensoryWare libs are compiled with libc++ instead of libstdc++
 				// https://udn.unrealengine.com/questions/305445/ue-33584-linux-build.html
diff --git a/Engine/Plugins/Lumin/MagicLeap/Source/MagicLeap/Private/AppEventHandler.cpp b/Engine/Plugins/Lumin/MagicLeap/Source/MagicLeap/Private/AppEventHandler.cpp
--- a/Engine/Plugins/Lumin/MagicLeap/Source/MagicLeap/Private/AppEventHandler.cpp
+++ b/Engine/Plugins/Lumin/MagicLeap/Source/MagicLeap/Private/AppEventHandler.cpp
@@ -4,53 +4,13 @@
 #include "AppFramework.h"
 #include "Engine/Engine.h"
 #include "MagicLeapHMD.h"
+#include "MagicLeapPrivilegeUtils.h"
 
 namespace MagicLeap
 {
-#if WITH_MLSDK
-	const TCHAR* MLPrivilegeToString(MLPrivilegeID PrivilegeID)
-	{
-		#define PRIV_TO_STR_CASE(x) case x: { return UTF8_TO_TCHAR((#x)); }
-		switch (PrivilegeID)
-		{
-		PRIV_TO_STR_CASE(MLPrivilegeID_Invalid)
-		PRIV_TO_STR_CASE(MLPrivilegeID_AudioRecognizer)
-		PRIV_TO_STR_CASE(MLPrivilegeID_BatteryInfo)
-		PRIV_TO_STR_CASE(MLPrivilegeID_CameraCapture)
-		PRIV_TO_STR_CASE(MLPrivilegeID_WorldReconstruction)
-		PRIV_TO_STR_CASE(MLPrivilegeID_InAppPurchase)
-		PRIV_TO_STR_CASE(MLPrivilegeID_AudioCaptureMic)
-		PRIV_TO_STR_CASE(MLPrivilegeID_DrmCertificates)
-		PRIV_TO_STR_CASE(MLPrivilegeID_Occlusion)
-		PRIV_TO_STR_CASE(MLPrivilegeID_LowLatencyLightwear)
-		PRIV_TO_STR_CASE(MLPrivilegeID_Internet)
-		PRIV_TO_STR_CASE(MLPrivilegeID_IdentityRead)
-		PRIV_TO_STR_CASE(MLPrivilegeID_BackgroundDownload)
-		PRIV_TO_STR_CASE(MLPrivilegeID_BackgroundUpload)
-		PRIV_TO_STR_CASE(MLPrivilegeID_MediaDrm)
-		PRIV_TO_STR_CASE(MLPrivilegeID_Media)
-		PRIV_TO_STR_CASE(MLPrivilegeID_MediaMetadata)
-		PRIV_TO_STR_CASE(MLPrivilegeID_PowerInfo)
-		PRIV_TO_STR_CASE(MLPrivilegeID_LocalAreaNetwork)
-		PRIV_TO_STR_CASE(MLPrivilegeID_VoiceInput)
-		PRIV_TO_STR_CASE(MLPrivilegeID_Documents)
-		PRIV_TO_STR_CASE(MLPrivilegeID_ConnectBackgroundMusicService)
-		PRIV_TO_STR_CASE(MLPrivilegeID_RegisterBackgroundMusicService)
-		PRIV_TO_STR_CASE(MLPrivilegeID_PwFoundObjRead)
-		PRIV_TO_STR_CASE(MLPrivilegeID_NormalNotificationsUsage)
-		PRIV_TO_STR_CASE(MLPrivilegeID_MusicService)
-		PRIV_TO_STR_CASE(MLPrivilegeID_ControllerPose)
-		PRIV_TO_STR_CASE(MLPrivilegeID_ScreensProvider)
-		PRIV_TO_STR_CASE(MLPrivilegeID_GesturesSubscribe)
-		PRIV_TO_STR_CASE(MLPrivilegeID_GesturesConfig)
-		default: UE_LOG(LogMagicLeap, Error, TEXT("Unmapped privilege %d"), static_cast<int32>(PrivilegeID));
-		}
-
-		return UTF8_TO_TCHAR("");
-	}
-
-	IAppEventHandler::IAppEventHandler(const TArray<MLPrivilegeID>& InRequiredPrivilegeIDs)
-	: OnAppShutDownHandler(nullptr)
+	IAppEventHandler::IAppEventHandler(const TArray<EMagicLeapPrivilege>& InRequiredPrivilegeIDs)
+	: OnPrivilegeEvent(nullptr)
+	, OnAppShutDownHandler(nullptr)
 	, OnAppTickHandler(nullptr)
 	, OnAppPauseHandler(nullptr)
 	, OnAppResumeHandler(nullptr)
@@ -58,14 +18,13 @@ namespace MagicLeap
 	, bWasSystemEnabledOnPause(false)
 	{
 		RequiredPrivileges.Reserve(InRequiredPrivilegeIDs.Num());
-		for (MLPrivilegeID RequiredPrivilegeID : InRequiredPrivilegeIDs)
+		for (EMagicLeapPrivilege RequiredPrivilegeID : InRequiredPrivilegeIDs)
 		{
 			RequiredPrivileges.Add(RequiredPrivilegeID, FRequiredPrivilege(RequiredPrivilegeID));
 		}
 
 		FAppFramework::AddEventHandler(this);
 	}
-#endif // WITH_MLSDK
 
 	IAppEventHandler::IAppEventHandler()
 	: bAllPrivilegesInSync(true)
@@ -79,46 +38,98 @@ namespace MagicLeap
 		FAppFramework::RemoveEventHandler(this);
 	}
 
-#if WITH_MLSDK
-	EPrivilegeState IAppEventHandler::GetPrivilegeStatus(MLPrivilegeID PrivilegeID, bool bBlocking/* = true*/)
+	EPrivilegeState IAppEventHandler::GetPrivilegeStatus(EMagicLeapPrivilege PrivilegeID, bool bBlocking/* = true*/)
 	{
-		FScopeLock Lock(&CriticalSection);
-		FRequiredPrivilege& RequiredPrivilege = RequiredPrivileges[PrivilegeID];
-		if (RequiredPrivilege.State == EPrivilegeState::NotYetRequested)
+#if WITH_MLSDK
+		FRequiredPrivilege RequiredPrivilege(PrivilegeID);
 		{
-			if (bBlocking)
+			FScopeLock Lock(&CriticalSection);
+			RequiredPrivilege = RequiredPrivileges[PrivilegeID];
+		}
+		
+		if (RequiredPrivilege.State == EPrivilegeState::NotYetRequested || RequiredPrivilege.State == EPrivilegeState::Pending)
+		{
+			// Attempt a quick check first. Note that MLPrivilegesCheckPrivilege will return denied if the privilege has never been requested.
+			MLResult Result = MLPrivilegesCheckPrivilege(MagicLeap::UnrealToMLPrivilege(PrivilegeID));
+			if (Result == MLPrivilegesResult_Granted)
+			{
+				RequiredPrivilege.State = EPrivilegeState::Granted;
+				UE_LOG(LogMagicLeap, Log, TEXT("Privilege '%s' was granted."), *MagicLeap::MLPrivilegeToString(RequiredPrivilege.PrivilegeID));
+			} 
+			else if (bBlocking)
 			{
-				MLResult Result = MLPrivilegesRequestPrivilege(PrivilegeID);
+				Result = MLPrivilegesRequestPrivilege(MagicLeap::UnrealToMLPrivilege(PrivilegeID));
 				switch (Result)
 				{
 				case MLPrivilegesResult_Granted:
 				{
 					RequiredPrivilege.State = EPrivilegeState::Granted;
-					UE_LOG(LogMagicLeap, Log, TEXT("Privilege '%s' was granted."), MLPrivilegeToString(RequiredPrivilege.PrivilegeID));
+					UE_LOG(LogMagicLeap, Log, TEXT("Privilege '%s' was granted."), *MagicLeap::MLPrivilegeToString(RequiredPrivilege.PrivilegeID));
 				}
 				break;
 				case MLPrivilegesResult_Denied:
 				{
 					RequiredPrivilege.State = EPrivilegeState::Denied;
-					UE_LOG(LogMagicLeap, Log, TEXT("Privilege '%s' was denied."), MLPrivilegeToString(RequiredPrivilege.PrivilegeID));
+					UE_LOG(LogMagicLeap, Log, TEXT("Privilege '%s' was denied."), *MagicLeap::MLPrivilegeToString(RequiredPrivilege.PrivilegeID));
 				}
 				break;
 				default:
 				{
 					UE_LOG(LogMagicLeap, Error, TEXT("MLPrivilegesRequestPrivilege() failed with error %s"), UTF8_TO_TCHAR(MLPrivilegesGetResultString(Result)));
+					RequiredPrivilege.State = EPrivilegeState::Error;
 				}
 				}
 			}
 			else
 			{
-				MLPrivilegesRequestPrivilegeAsync(PrivilegeID, &RequiredPrivilege.PrivilegeRequest);
-				RequiredPrivilege.State = EPrivilegeState::Pending;
+				Result = MLPrivilegesRequestPrivilegeAsync(MagicLeap::UnrealToMLPrivilege(PrivilegeID), reinterpret_cast<MLPrivilegesAsyncRequest**>(&RequiredPrivilege.PrivilegeRequest));
+				if (Result != MLResult_Ok)
+				{
+					UE_LOG(LogMagicLeap, Error, TEXT("MLPrivilegesRequestPrivilegeAsync() failed with error %s"), UTF8_TO_TCHAR(MLPrivilegesGetResultString(Result)));
+					RequiredPrivilege.State = EPrivilegeState::Error;
+				}
+				else
+				{
+					RequiredPrivilege.State = EPrivilegeState::Pending;
+				}
 			}
 		}
 
+		{
+			FScopeLock Lock(&CriticalSection);
+			RequiredPrivileges[PrivilegeID] = RequiredPrivilege;
+		}
+
 		return RequiredPrivilege.State;
+#else
+		return EPrivilegeState::NotYetRequested;
+#endif // WITH_MLSDK
 	}
+
+	FString IAppEventHandler::PrivilegeToString(EMagicLeapPrivilege PrivilegeID)
+	{
+#if WITH_MLSDK
+		return MLPrivilegeToString(PrivilegeID);
+#else
+		return TEXT("MLPrivilegeID_Invalid");
 #endif // WITH_MLSDK
+	}
+
+	const TCHAR* IAppEventHandler::PrivilegeStateToString(EPrivilegeState PrivilegeState)
+	{
+		const TCHAR* PrivilegeStateString = nullptr;
+
+		switch (PrivilegeState)
+		{
+		case EPrivilegeState::NotYetRequested:	PrivilegeStateString = TEXT("NotYetRequested"); break;
+		case EPrivilegeState::Pending:			PrivilegeStateString = TEXT("Pending"); break;
+		case EPrivilegeState::Granted:			PrivilegeStateString = TEXT("Granted"); break;
+		case EPrivilegeState::Denied:			PrivilegeStateString = TEXT("Denied"); break;
+		case EPrivilegeState::Error:			PrivilegeStateString = TEXT("Error"); break;
+		}
+
+		return PrivilegeStateString;
+	}
 
 	void IAppEventHandler::OnAppShutDown()
 	{
@@ -155,19 +166,19 @@ namespace MagicLeap
 
 			if (RequiredPrivilege.State == EPrivilegeState::Pending)
 			{
-				MLResult Result = MLPrivilegesRequestPrivilegeTryGet(RequiredPrivilege.PrivilegeRequest);
+				MLResult Result = MLPrivilegesRequestPrivilegeTryGet(static_cast<MLPrivilegesAsyncRequest*>(RequiredPrivilege.PrivilegeRequest));
 				switch (Result)
 				{
 				case MLPrivilegesResult_Granted:
 				{
 					RequiredPrivilege.State = EPrivilegeState::Granted;
-					UE_LOG(LogMagicLeap, Log, TEXT("Privilege '%s' was granted."), MLPrivilegeToString(RequiredPrivilege.PrivilegeID));
+					UE_LOG(LogMagicLeap, Log, TEXT("Privilege '%s' was granted."), *MagicLeap::MLPrivilegeToString(RequiredPrivilege.PrivilegeID));
 				}
 				break;
 				case MLPrivilegesResult_Denied:
 				{
 					RequiredPrivilege.State = EPrivilegeState::Denied;
-					UE_LOG(LogMagicLeap, Log, TEXT("Privilege '%s' was denied."), MLPrivilegeToString(RequiredPrivilege.PrivilegeID));
+					UE_LOG(LogMagicLeap, Log, TEXT("Privilege '%s' was denied."), *MagicLeap::MLPrivilegeToString(RequiredPrivilege.PrivilegeID));
 				}
 				break;
 				case MLResult_Pending:
@@ -178,8 +189,13 @@ namespace MagicLeap
 				default:
 				{
 					bAllPrivilegesInSync = false;
-					UE_LOG(LogMagicLeap, Error, TEXT("MLPrivilegesRequestPrivilegeTryGet() failed with error %s"), UTF8_TO_TCHAR(MLPrivilegesGetResultString(Result)));
+					UE_LOG(LogMagicLeap, Error, TEXT("MLPrivilegesRequestPrivilegeTryGet() failed with error %s."), UTF8_TO_TCHAR(MLPrivilegesGetResultString(Result)));
+				}
 				}
+
+				if (Result != MLResult_Pending && OnPrivilegeEvent)
+				{
+					OnPrivilegeEvent(RequiredPrivilege);
 				}
 			}
 		}
@@ -220,4 +236,4 @@ namespace MagicLeap
 	{
 		return FAppFramework::AsyncDestroy(this);
 	}
-}
\ No newline at end of file
+}
diff --git a/Engine/Plugins/Lumin/MagicLeap/Source/MagicLeap/Private/AppFramework.cpp b/Engine/Plugins/Lumin/MagicLeap/Source/MagicLeap/Private/AppFramework.cpp
--- a/Engine/Plugins/Lumin/MagicLeap/Source/MagicLeap/Private/AppFramework.cpp
+++ b/Engine/Plugins/Lumin/MagicLeap/Source/MagicLeap/Private/AppFramework.cpp
@@ -9,19 +9,16 @@
 #include "Misc/CoreDelegates.h"
 #include "RenderingThread.h"
 #include "MagicLeapPluginUtil.h" // for ML_INCLUDES_START/END
-#include "IMagicLeapModule.h"
 
-#if WITH_MLSDK
-ML_INCLUDES_START
-#include <ml_snapshot.h>
-ML_INCLUDES_END
-#endif //WITH_MLSDK
+#include "Lumin/CAPIShims/LuminAPISnapshot.h"
+
+#if PLATFORM_LUMIN
+#include "Lumin/LuminPlatformDelegates.h"
+#endif // PLATFORM_LUMIN
 
 TArray<MagicLeap::IAppEventHandler*> FAppFramework::EventHandlers;
 FCriticalSection FAppFramework::EventHandlersCriticalSection;
 MagicLeap::FAsyncDestroyer* FAppFramework::AsyncDestroyer = nullptr;
-TMap<FName, IMagicLeapModule*> FAppFramework::RegisteredModules;
-
 
 FAppFramework::FAppFramework()
 {}
@@ -33,20 +30,16 @@ FAppFramework::~FAppFramework()
 
 void FAppFramework::Startup()
 {
-	base_dirty_ = false;
-
-#if WITH_MLSDK
-	base_coordinate_frame_.data[0] = 0;
-	base_coordinate_frame_.data[1] = 0;
-#endif //WITH_MLSDK
-
-	base_position_ = FVector::ZeroVector;
-	base_orientation_ = FQuat::Identity;
-
 	// Register application lifecycle delegates
 	FCoreDelegates::ApplicationWillEnterBackgroundDelegate.AddRaw(this, &FAppFramework::ApplicationPauseDelegate);
 	FCoreDelegates::ApplicationHasEnteredForegroundDelegate.AddRaw(this, &FAppFramework::ApplicationResumeDelegate);
 
+#if PLATFORM_LUMIN
+	FLuminDelegates::DeviceHasReactivatedDelegate.AddRaw(this, &FAppFramework::OnDeviceActive);
+	FLuminDelegates::DeviceWillEnterRealityModeDelegate.AddRaw(this, &FAppFramework::OnDeviceRealityMode);
+	FLuminDelegates::DeviceWillGoInStandbyDelegate.AddRaw(this, &FAppFramework::OnDeviceStandby);
+#endif // PLATFORM_LUMIN
+
 	AsyncDestroyer = new MagicLeap::FAsyncDestroyer();
 
 	bInitialized = true;
@@ -60,6 +53,15 @@ void FAppFramework::Shutdown()
 
 	delete AsyncDestroyer;
 	AsyncDestroyer = nullptr;
+
+	FCoreDelegates::ApplicationWillEnterBackgroundDelegate.RemoveAll(this);
+	FCoreDelegates::ApplicationHasEnteredForegroundDelegate.RemoveAll(this);
+
+#if PLATFORM_LUMIN
+	FLuminDelegates::DeviceHasReactivatedDelegate.RemoveAll(this);
+	FLuminDelegates::DeviceWillEnterRealityModeDelegate.RemoveAll(this);
+	FLuminDelegates::DeviceWillGoInStandbyDelegate.RemoveAll(this);
+#endif // PLATFORM_LUMIN
 }
 
 void FAppFramework::BeginUpdate()
@@ -67,19 +69,6 @@ void FAppFramework::BeginUpdate()
 #if WITH_MLSDK
 	if (bInitialized)
 	{
-		if (base_dirty_)
-		{
-			MLTransform Transform = MagicLeap::kIdentityTransform;
-			Transform.position = MagicLeap::ToMLVector(-base_position_, GetWorldToMetersScale());
-			Transform.rotation = MagicLeap::ToMLQuat(base_orientation_.Inverse());
-
-			base_coordinate_frame_.data[0] = 0;
-			base_coordinate_frame_.data[1] = 0;
-			base_position_ = FVector::ZeroVector;
-			base_orientation_ = FQuat::Identity;
-			base_dirty_ = false;
-		}
-
 		FScopeLock Lock(&EventHandlersCriticalSection);
 		for (auto EventHandler : EventHandlers)
 		{
@@ -113,11 +102,7 @@ void FAppFramework::ApplicationPauseDelegate()
 	}
 
 	// Pause rendering
-	FMagicLeapHMD * const HMD = GEngine ? static_cast<FMagicLeapHMD*>(GEngine->XRSystem->GetHMDDevice()) : nullptr;
-	if (HMD)
-	{
-		HMD->PauseRendering(true);
-	}
+	PauseRendering(true);
 }
 
 void FAppFramework::ApplicationResumeDelegate()
@@ -125,12 +110,7 @@ void FAppFramework::ApplicationResumeDelegate()
 	UE_LOG(LogMagicLeap, Log, TEXT("+++++++ ML AppFramework APP RESUME ++++++"));
 
 	// Resume rendering
-	// Resume rendering
-	FMagicLeapHMD * const HMD = GEngine ? static_cast<FMagicLeapHMD*>(GEngine->XRSystem->GetHMDDevice()) : nullptr;
-	if (HMD)
-	{
-		HMD->PauseRendering(false);
-	}
+	PauseRendering(false);
 
 	if (GEngine)
 	{
@@ -159,30 +139,31 @@ void FAppFramework::OnApplicationShutdown()
 	}
 }
 
-#if WITH_MLSDK
-void FAppFramework::SetBaseCoordinateFrame(MLCoordinateFrameUID InBaseCoordinateFrame)
+void FAppFramework::OnDeviceActive()
 {
-	base_coordinate_frame_ = InBaseCoordinateFrame;
-	base_dirty_ = true;
+	UE_LOG(LogMagicLeap, Log, TEXT("+++++++ ML AppFramework DEVICE ACTIVE ++++++"));
+	PauseRendering(false);
 }
-#endif //WITH_MLSDK
 
-void FAppFramework::SetBasePosition(const FVector& InBasePosition)
+void FAppFramework::OnDeviceRealityMode()
 {
-	base_position_ = InBasePosition;
-	base_dirty_ = true;
+	UE_LOG(LogMagicLeap, Log, TEXT("+++++++ ML AppFramework DEVICE REALITY MODE ++++++"));
+	PauseRendering(true);
 }
 
-void FAppFramework::SetBaseOrientation(const FQuat& InBaseOrientation)
+void FAppFramework::OnDeviceStandby()
 {
-	base_orientation_ = InBaseOrientation;
-	base_dirty_ = true;
+	UE_LOG(LogMagicLeap, Log, TEXT("+++++++ ML AppFramework DEVICE STANDBY ++++++"));
+	PauseRendering(true);
 }
 
-void FAppFramework::SetBaseRotation(const FRotator& InBaseRotation)
+void FAppFramework::PauseRendering(bool bPause)
 {
-	base_orientation_ = InBaseRotation.Quaternion();
-	base_dirty_ = true;
+	FMagicLeapHMD * const HMD = GEngine ? static_cast<FMagicLeapHMD*>(GEngine->XRSystem->GetHMDDevice()) : nullptr;
+	if (HMD)
+	{
+		HMD->PauseRendering(bPause);
+	}
 }
 
 const FTrackingFrame* FAppFramework::GetCurrentFrame() const
@@ -197,17 +178,11 @@ const FTrackingFrame* FAppFramework::GetOldFrame() const
 	return hmd ? &(hmd->GetOldFrame()) : nullptr;
 }
 
-FVector2D FAppFramework::GetFieldOfView() const
-{
-	// TODO Pass correct values when graphics provides them through API.
-	return FVector2D(80.0f, 60.0f);
-}
-
 uint32 FAppFramework::GetViewportCount() const
 {
 #if WITH_MLSDK
 	const FTrackingFrame *frame = GetOldFrame();
-	return frame ? frame->RenderInfoArray.num_virtual_cameras : 2;
+	return frame ? frame->FrameInfo.virtual_camera_info_array.num_virtual_cameras : 2;
 #else
 	return 1;
 #endif //WITH_MLSDK
@@ -233,12 +208,12 @@ FTransform FAppFramework::GetCurrentFrameUpdatePose() const
 }
 
 #if WITH_MLSDK
-bool FAppFramework::GetTransform(const MLCoordinateFrameUID& Id, FTransform& OutTransform, EFailReason& OutReason) const
+bool FAppFramework::GetTransform(const MLCoordinateFrameUID& Id, FTransform& OutTransform, EMagicLeapTransformFailReason& OutReason) const
 {
 	const FTrackingFrame* frame = GetCurrentFrame();
 	if (frame == nullptr)
 	{
-		OutReason = EFailReason::InvalidTrackingFrame;
+		OutReason = EMagicLeapTransformFailReason::InvalidTrackingFrame;
 		return false;
 	}
 
@@ -249,7 +224,7 @@ bool FAppFramework::GetTransform(const MLCoordinateFrameUID& Id, FTransform& Out
 		OutTransform = MagicLeap::ToFTransform(transform, GetWorldToMetersScale());
 		if (OutTransform.ContainsNaN())
 		{
-			OutReason = EFailReason::NaNsInTransform;
+			OutReason = EMagicLeapTransformFailReason::NaNsInTransform;
 			return false;
 		}
 		// Unreal crashes if the incoming quaternion is not normalized.
@@ -259,59 +234,22 @@ bool FAppFramework::GetTransform(const MLCoordinateFrameUID& Id, FTransform& Out
 			rotation.Normalize();
 			OutTransform.SetRotation(rotation);
 		}
-		OutReason = EFailReason::None;
+		OutReason = EMagicLeapTransformFailReason::None;
 		return true;
 	}
 	else if (Result == MLSnapshotResult_PoseNotFound)
 	{
-		OutReason = EFailReason::PoseNotFound;
+		OutReason = EMagicLeapTransformFailReason::PoseNotFound;
 	}
 	else
 	{
-		OutReason = EFailReason::CallFailed;
+		OutReason = EMagicLeapTransformFailReason::CallFailed;
 	}
 
 	return false;
 }
 #endif //WITH_MLSDK
 
-TSharedPtr<FCameraCaptureRunnable, ESPMode::ThreadSafe> FAppFramework::GetCameraCaptureRunnable()
-{
-	if (!CameraCaptureRunnable.IsValid())
-	{
-		CameraCaptureRunnable = MakeShared<FCameraCaptureRunnable, ESPMode::ThreadSafe>();
-	}
-	return CameraCaptureRunnable;
-}
-
-void FAppFramework::RefreshCameraCaptureRunnableReferences()
-{
-	// a reference count of 1 is a self reference
-	if (CameraCaptureRunnable.GetSharedReferenceCount() == 1)
-	{
-		CameraCaptureRunnable.Reset();
-	}
-}
-
-TSharedPtr<FImageTrackerRunnable, ESPMode::ThreadSafe> FAppFramework::GetImageTrackerRunnable()
-{
-	if (!ImageTrackerRunnable.IsValid())
-	{
-		ImageTrackerRunnable = MakeShared<FImageTrackerRunnable, ESPMode::ThreadSafe>();
-	}
-
-	return ImageTrackerRunnable;
-}
-
-void FAppFramework::RefreshImageTrackerRunnableReferences()
-{
-	// a reference count of 1 is a self reference
-	if (ImageTrackerRunnable.GetSharedReferenceCount() == 1)
-	{
-		ImageTrackerRunnable.Reset();
-	}
-}
-
 void FAppFramework::AddEventHandler(MagicLeap::IAppEventHandler* EventHandler)
 {
 	FScopeLock Lock(&EventHandlersCriticalSection);
@@ -334,20 +272,3 @@ bool FAppFramework::AsyncDestroy(MagicLeap::IAppEventHandler* InEventHandler)
 
 	return false;
 }
-
-void FAppFramework::RegisterMagicLeapModule(IMagicLeapModule* InModule)
-{
-	checkf(RegisteredModules.Find(InModule->GetName()) == nullptr, TEXT("MagicLeapModule %s has already been registered!"), *InModule->GetName().ToString());
-	RegisteredModules.Add(InModule->GetName(), InModule);
-}
-
-void FAppFramework::UnregisterMagicLeapModule(IMagicLeapModule* InModule)
-{
-	RegisteredModules.Remove(InModule->GetName());
-}
-
-IMagicLeapModule* FAppFramework::GetMagicLeapModule(FName InName)
-{
-	IMagicLeapModule** MagicLeapModule = RegisteredModules.Find(InName);
-	return MagicLeapModule ? *MagicLeapModule : nullptr;
-}
\ No newline at end of file
diff --git a/Engine/Plugins/Lumin/MagicLeap/Source/MagicLeap/Private/AppFramework.h b/Engine/Plugins/Lumin/MagicLeap/Source/MagicLeap/Private/AppFramework.h
--- a/Engine/Plugins/Lumin/MagicLeap/Source/MagicLeap/Private/AppFramework.h
+++ b/Engine/Plugins/Lumin/MagicLeap/Source/MagicLeap/Private/AppFramework.h
@@ -3,36 +3,16 @@
 #pragma once
 
 #include "GenericPlatform/GenericPlatformProcess.h"
-#include "MagicLeapPluginUtil.h" // for ML_INCLUDES_START/END
-
-#if WITH_MLSDK
-ML_INCLUDES_START
-
-#include <ml_api.h>
-#include <ml_coordinate_frame_uid.h>
-
-ML_INCLUDES_END
-#endif  //WITH_MLSDK
 
+#include "Lumin/CAPIShims/LuminAPICoordinateFrameUID.h"
 #include "CoreMinimal.h"
 #include "AppEventHandler.h"
 #include "AsyncDestroyer.h"
-#include "CameraCaptureRunnable.h"
-#include "ImageTrackerRunnable.h"
+#include "IMagicLeapPlugin.h"
 
 class FMagicLeapHMD;
-class IMagicLeapModule;
 struct FTrackingFrame;
 
-enum class EFailReason : uint8
-{
-	None,
-	InvalidTrackingFrame,
-	NaNsInTransform,
-	CallFailed,
-	PoseNotFound
-};
-
 class MAGICLEAP_API FAppFramework
 {
 public:
@@ -44,27 +24,13 @@ public:
 
 	void BeginUpdate();
 
-	bool IsInitialized() const { return bInitialized; }
-
 	void ApplicationPauseDelegate();
 	void ApplicationResumeDelegate();
 	void OnApplicationShutdown();
 
-#if WITH_MLSDK
-	void SetBaseCoordinateFrame(MLCoordinateFrameUID InBaseCoordinateFrame);
-	MLCoordinateFrameUID GetBaseCoordinateFrame() const { return base_coordinate_frame_; }
-#endif //WITH_MLSDK
-
-	void SetBasePosition(const FVector& InBasePosition);
-	FVector GetBasePosition() const { return base_position_; }
-
-	void SetBaseOrientation(const FQuat& InBaseOrientation);
-	FQuat GetBaseOrientation() const { return base_orientation_; }
-
-	void SetBaseRotation(const FRotator& InBaseRotation);
-	FRotator GetBaseRotation() const { return FRotator(base_orientation_); }
-
-	FVector2D GetFieldOfView() const;
+	void OnDeviceActive();
+	void OnDeviceRealityMode();
+	void OnDeviceStandby();
 
 	FTransform GetDisplayCenterTransform() const { return FTransform::Identity; }; // HACK
 	uint32 GetViewportCount() const;
@@ -72,44 +38,26 @@ public:
 	float GetWorldToMetersScale() const;
 	FTransform GetCurrentFrameUpdatePose() const;
 #if WITH_MLSDK
-	bool GetTransform(const MLCoordinateFrameUID& Id, FTransform& OutTransform, EFailReason& OutReason) const;
+	bool GetTransform(const MLCoordinateFrameUID& Id, FTransform& OutTransform, EMagicLeapTransformFailReason& OutReason) const;
 #endif //WITH_MLSDK
 
-	TSharedPtr<FCameraCaptureRunnable, ESPMode::ThreadSafe> GetCameraCaptureRunnable();
-	void RefreshCameraCaptureRunnableReferences();
-	TSharedPtr<FImageTrackerRunnable, ESPMode::ThreadSafe> GetImageTrackerRunnable();
-	void RefreshImageTrackerRunnableReferences();
-
 	static void AddEventHandler(MagicLeap::IAppEventHandler* InEventHandler);
 	static void RemoveEventHandler(MagicLeap::IAppEventHandler* InEventHandler);
 	static bool AsyncDestroy(MagicLeap::IAppEventHandler* InEventHandler);
 
-	static void RegisterMagicLeapModule(IMagicLeapModule* InModule);
-	static void UnregisterMagicLeapModule(IMagicLeapModule* InModule);
-	static IMagicLeapModule* GetMagicLeapModule(FName InName);
-
 private:
 	const FTrackingFrame* GetCurrentFrame() const;
 	const FTrackingFrame* GetOldFrame() const;
+	void PauseRendering(bool bPause);
 
 	bool bInitialized = false;
 
-#if WITH_MLSDK
-	MLCoordinateFrameUID base_coordinate_frame_;
-#endif //WITH_MLSDK
-
-	FVector base_position_;
-	FQuat base_orientation_;
-	bool base_dirty_;
-
 	float saved_max_fps_;
 
 	static TArray<MagicLeap::IAppEventHandler*> EventHandlers;
 	static FCriticalSection EventHandlersCriticalSection;
 	static MagicLeap::FAsyncDestroyer* AsyncDestroyer;
-	TSharedPtr<FCameraCaptureRunnable, ESPMode::ThreadSafe> CameraCaptureRunnable;
-	TSharedPtr<FImageTrackerRunnable, ESPMode::ThreadSafe> ImageTrackerRunnable;
-	static TMap<FName, IMagicLeapModule*> RegisteredModules;
 };
 
-DEFINE_LOG_CATEGORY_STATIC(LogMagicLeap, Log, All);
\ No newline at end of file
+// TODO: pull this out of here
+DEFINE_LOG_CATEGORY_STATIC(LogMagicLeap, Log, All);
