Speculative fix for null reference exception when Lumin is not enabled in a binary build/C++ project.

diff --git a/Engine/Plugins/Lumin/MLSDK/Source/MLSDK/MLSDK.Build.cs b/Engine/Plugins/Lumin/MLSDK/Source/MLSDK/MLSDK.Build.cs
--- a/Engine/Plugins/Lumin/MLSDK/Source/MLSDK/MLSDK.Build.cs
+++ b/Engine/Plugins/Lumin/MLSDK/Source/MLSDK/MLSDK.Build.cs
@@ -17,7 +17,7 @@ public class MLSDK : ModuleRules
 		string MLSDKPath = System.Environment.GetEnvironmentVariable("MLSDK");
 		bool bIsMLSDKInstalled = false;
 		bool bIsSupportedTargetPlatformAndType = ((Target.Platform == UnrealTargetPlatform.Lumin) || (Target.Type == TargetRules.TargetType.Editor));
-		if (SDK.IsVersionValid(SDK.GetInstalledVersion(), bForAutoSDK:false) && !string.IsNullOrEmpty(MLSDKPath))
+		if (SDK != null && SDK.IsVersionValid(SDK.GetInstalledVersion(), bForAutoSDK:false) && !string.IsNullOrEmpty(MLSDKPath))
 		{
 			string IncludePath = Path.Combine(MLSDKPath, "include");
 			string LibraryPath = Path.Combine(MLSDKPath, "lib");
